UAF: Added UAF Layering plugin

- New Layer Stack asset
- Workspace Integration for asset
- First iteration of layer stack editor and runtime

#Synced-CL 49758930

diff --git a/Engine/Plugins/Experimental/UAF/UAFLayering/Source/UAFLayering/Internal/Traits/CachePoseTraitData.h b/Engine/Plugins/Experimental/UAF/UAFLayering/Source/UAFLayering/Internal/Traits/CachePoseTraitData.h
new file mode 100644
--- /dev/null
+++ b/Engine/Plugins/Experimental/UAF/UAFLayering/Source/UAFLayering/Internal/Traits/CachePoseTraitData.h
@@ -0,0 +1,18 @@
+// Copyright Epic Games, Inc. All Rights Reserved.
+
+#pragma once
+
+#include "TraitCore/TraitBinding.h"
+#include "TraitCore/TraitSharedData.h"
+
+#include "CachePoseTraitData.generated.h"
+
+/** A trait that caches the top keyframe of the keyframe stack
+ * TODO: Not implemented yet, at the moment it will simply pop the top keyframe and throw it away*/
+USTRUCT(meta = (DisplayName = "Cache Pose", ShowTooltip=true, Hidden))
+struct FUAFCachePoseTraitSharedData : public FAnimNextTraitSharedData
+{
+	GENERATED_BODY()
+};
+
+
diff --git a/Engine/Plugins/Experimental/UAF/UAFLayering/Source/UAFLayering/Internal/Traits/LayerAssetDataTraitData.h b/Engine/Plugins/Experimental/UAF/UAFLayering/Source/UAFLayering/Internal/Traits/LayerAssetDataTraitData.h
new file mode 100644
--- /dev/null
+++ b/Engine/Plugins/Experimental/UAF/UAFLayering/Source/UAFLayering/Internal/Traits/LayerAssetDataTraitData.h
@@ -0,0 +1,30 @@
+// Copyright Epic Games, Inc. All Rights Reserved.
+
+#pragma once
+#include "Factory/AnimNextFactoryParams.h"
+#include "TraitCore/TraitBinding.h"
+#include "TraitCore/TraitSharedData.h"
+
+#include "LayerAssetDataTraitData.generated.h"
+
+
+/** A trait that creates a anim graph at runtime based on asset and factory params of the layering setup and pushes it onto a blend stack*/
+USTRUCT(meta = (DisplayName = "Layer Data Provider", ShowTooltip=true, Hidden))
+struct FUAFLayerAssetDataTraitSharedData : public FAnimNextTraitSharedData
+{
+	GENERATED_BODY()
+	
+	UPROPERTY(EditAnywhere, Category = "Layer|Content", meta = (GetAllowedClasses = "/Script/UAFAnimGraph.AnimNextAnimGraphSettings:GetAllowedAssetClasses"))
+	TObjectPtr<UObject> Asset = nullptr;
+
+	UPROPERTY(EditAnywhere, Category = "Layer|Content", meta=(FactorySourceProperty = "Asset"))
+	FAnimNextFactoryParams Parameters;
+	
+	// Latent pin support boilerplate
+	#define TRAIT_LATENT_PROPERTIES_ENUMERATOR(GeneratorMacro) \
+	GeneratorMacro(Asset) 
+
+	GENERATE_TRAIT_LATENT_PROPERTIES(FUAFLayerAssetDataTraitSharedData, TRAIT_LATENT_PROPERTIES_ENUMERATOR)
+	#undef TRAIT_LATENT_PROPERTIES_ENUMERATOR
+	
+};
diff --git a/Engine/Plugins/Experimental/UAF/UAFLayering/Source/UAFLayering/Internal/Traits/LayerDataProviderTraitData.h b/Engine/Plugins/Experimental/UAF/UAFLayering/Source/UAFLayering/Internal/Traits/LayerDataProviderTraitData.h
new file mode 100644
--- /dev/null
+++ b/Engine/Plugins/Experimental/UAF/UAFLayering/Source/UAFLayering/Internal/Traits/LayerDataProviderTraitData.h
@@ -0,0 +1,87 @@
+// Copyright Epic Games, Inc. All Rights Reserved.
+
+#pragma once
+
+#include "AlphaBlend.h"
+#include "UAFLayeringTypes.h"
+#include "TraitCore/TraitBinding.h"
+#include "TraitCore/TraitSharedData.h"
+
+#include "LayerDataProviderTraitData.generated.h"
+
+class UCurveFloat;
+
+USTRUCT(BlueprintType)
+struct FUAFLayerProperties
+{
+	GENERATED_BODY()
+
+public:
+	void AddEvent(const UE::UAF::Layering::FLayerStack_LayerEvent& Event);
+	void ProcessEvents();
+	
+public:
+	// TODO: this should be a binding 
+	UPROPERTY(VisibleAnywhere, BlueprintReadOnly, Category = "Layering")
+	bool bLayerEnabled = true;
+		
+	// TODO: this should be a binding 
+	UPROPERTY(VisibleAnywhere, BlueprintReadOnly, Category = "Layering")
+	float DesiredLayerWeight = 1.0f;
+		
+	// TODO: this should be a binding 
+	UPROPERTY(VisibleAnywhere, BlueprintReadOnly, Category = "Layering")
+	float BlendInTime = 0.0f;
+		
+	// TODO: this should be a binding 
+	UPROPERTY(VisibleAnywhere, BlueprintReadOnly, Category = "Layering")
+	float BlendOutTime = 0.0f;
+	
+	// TODO: this should be a binding 
+	UPROPERTY(VisibleAnywhere, BlueprintReadOnly, Category = "Layering")
+	TObjectPtr<UCurveFloat> BlendCurve = nullptr;
+	
+	// TODO: this should be a binding 
+	UPROPERTY(VisibleAnywhere, BlueprintReadOnly, Category = "Layering")
+	EAlphaBlendOption BlendOption = EAlphaBlendOption::Linear;
+
+private:
+	TArray<UE::UAF::Layering::FLayerStack_LayerEvent> EventsToProcess;
+};
+
+/** A trait that provides blend data based the outer layer stack */
+USTRUCT(meta = (DisplayName = "Layer Data Provider", ShowTooltip=true, Hidden))
+struct FUAFLayerDataProviderTraitSharedData : public FAnimNextTraitSharedData
+{
+	GENERATED_BODY()
+
+	UPROPERTY(EditAnywhere, Category = "Layer")
+	FUAFLayerProperties LayerProperties;
+
+	// The index of this layer in its owning layer stack
+	UPROPERTY(EditAnywhere, Category = "Layer")
+	int32 LayerIndex = INDEX_NONE;
+
+	// The layer name of this layer 
+	UPROPERTY(EditAnywhere, Category = "Layer")
+	FName LayerName = NAME_None;
+	
+	// The soft object path to the owning layer stack
+	UPROPERTY(EditAnywhere, Category = "Layer")
+	FString LayerStackPath;
+	
+	// If this is true we will allocate the input to allow for update/eval of cache only branches of the tree
+	bool bCreateCacheInput = false;
+
+	// This is only used to enqueue Cache Only layers to the graph update without actually blending them in in-place 
+	UPROPERTY()
+	FAnimNextTraitHandle CacheOnlyInput;
+
+	// Latent pin support boilerplate
+	#define TRAIT_LATENT_PROPERTIES_ENUMERATOR(GeneratorMacro) \
+	GeneratorMacro(LayerProperties)
+
+	GENERATE_TRAIT_LATENT_PROPERTIES(FUAFLayerDataProviderTraitSharedData, TRAIT_LATENT_PROPERTIES_ENUMERATOR)
+	#undef TRAIT_LATENT_PROPERTIES_ENUMERATOR
+
+};
diff --git a/Engine/Plugins/Experimental/UAF/UAFLayering/Source/UAFLayering/Internal/UAFLayerStack.h b/Engine/Plugins/Experimental/UAF/UAFLayering/Source/UAFLayering/Internal/UAFLayerStack.h
new file mode 100644
--- /dev/null
+++ b/Engine/Plugins/Experimental/UAF/UAFLayering/Source/UAFLayering/Internal/UAFLayerStack.h
@@ -0,0 +1,15 @@
+// Copyright Epic Games, Inc. All Rights Reserved.
+
+#pragma once
+#include "Graph/AnimNextAnimationGraph.h"
+#include "UAFLayerStack.generated.h"
+
+UCLASS(BlueprintType, MinimalAPI)
+class UUAFLayerStack : public UUAFAnimGraph
+{
+	GENERATED_BODY()
+
+public:
+	friend class UUAFLayerStack_EditorData;
+	friend class UUAFLayerStackFactory;
+};
diff --git a/Engine/Plugins/Experimental/UAF/UAFLayering/Source/UAFLayering/Private/Tasks/UAFCachePoseTask.cpp b/Engine/Plugins/Experimental/UAF/UAFLayering/Source/UAFLayering/Private/Tasks/UAFCachePoseTask.cpp
new file mode 100644
--- /dev/null
+++ b/Engine/Plugins/Experimental/UAF/UAFLayering/Source/UAFLayering/Private/Tasks/UAFCachePoseTask.cpp
@@ -0,0 +1,27 @@
+// Copyright Epic Games, Inc. All Rights Reserved.
+
+
+
+#include "Tasks/UAFCachePoseTask.h"
+
+#include "EvaluationVM/EvaluationVM.h"
+#include "EvaluationVM/KeyframeState.h"
+
+DEFINE_STAT(STAT_AnimNext_Task_CachePose);
+
+FUAFCachePoseTask FUAFCachePoseTask::Make()
+{
+	return FUAFCachePoseTask();
+}
+
+void FUAFCachePoseTask::Execute(UE::UAF::FEvaluationVM& VM) const
+{
+	using namespace UE::UAF;
+
+	SCOPE_CYCLE_COUNTER(STAT_AnimNext_Task_CachePose);
+
+	TUniquePtr<FKeyframeState> PoseToCache;
+	const bool bCachedPose = VM.PopValue(KEYFRAME_STACK_NAME, PoseToCache);
+
+	// TODO: Write that pose to a heap allocated property or stack memory
+}
diff --git a/Engine/Plugins/Experimental/UAF/UAFLayering/Source/UAFLayering/Private/Tasks/UAFCachePoseTask.h b/Engine/Plugins/Experimental/UAF/UAFLayering/Source/UAFLayering/Private/Tasks/UAFCachePoseTask.h
new file mode 100644
--- /dev/null
+++ b/Engine/Plugins/Experimental/UAF/UAFLayering/Source/UAFLayering/Private/Tasks/UAFCachePoseTask.h
@@ -0,0 +1,30 @@
+// Copyright Epic Games, Inc. All Rights Reserved.
+
+#pragma once
+
+#include "CoreMinimal.h"
+#include "AnimNextStats.h"
+#include "EvaluationVM/EvaluationTask.h"
+#include "Stats/Stats.h"
+
+#include "UAFCachePoseTask.generated.h"
+
+
+DECLARE_CYCLE_STAT_EXTERN(TEXT("UAF Task: CachePose"), STAT_AnimNext_Task_CachePose, STATGROUP_AnimNext, UAFLAYERING_API);
+
+#define UE_API UAFLAYERING_API
+
+USTRUCT()
+struct FUAFCachePoseTask : public FAnimNextEvaluationTask
+{
+	GENERATED_BODY()
+
+	DECLARE_ANIM_EVALUATION_TASK(FUAFCachePoseTask)
+	
+	static UE_API FUAFCachePoseTask Make();
+	
+	// Task Entry point
+	UE_API virtual void Execute(UE::UAF::FEvaluationVM& VM) const override;
+};
+
+#undef UE_API
diff --git a/Engine/Plugins/Experimental/UAF/UAFLayering/Source/UAFLayering/Private/Traits/CachePoseTrait.cpp b/Engine/Plugins/Experimental/UAF/UAFLayering/Source/UAFLayering/Private/Traits/CachePoseTrait.cpp
new file mode 100644
--- /dev/null
+++ b/Engine/Plugins/Experimental/UAF/UAFLayering/Source/UAFLayering/Private/Traits/CachePoseTrait.cpp
@@ -0,0 +1,25 @@
+// Copyright Epic Games, Inc. All Rights Reserved.
+
+
+#include "CachePoseTrait.h"
+#include "Tasks/UAFCachePoseTask.h"
+
+namespace UE::UAF
+{
+AUTO_REGISTER_ANIM_TRAIT(FCachePoseTrait)
+
+// Trait implementation boilerplate
+#define TRAIT_INTERFACE_ENUMERATOR(GeneratorMacro) \
+GeneratorMacro(IEvaluate) \
+
+GENERATE_ANIM_TRAIT_IMPLEMENTATION(FCachePoseTrait, TRAIT_INTERFACE_ENUMERATOR, NULL_ANIM_TRAIT_INTERFACE_ENUMERATOR, NULL_ANIM_TRAIT_EVENT_ENUMERATOR)
+#undef TRAIT_INTERFACE_ENUMERATOR
+
+	
+void FCachePoseTrait::PostEvaluate(FEvaluateTraversalContext& Context, const TTraitBinding<IEvaluate>& Binding) const
+{	
+	Context.AppendTask(FUAFCachePoseTask::Make());
+	
+	IEvaluate::PostEvaluate(Context, Binding);
+}
+}
diff --git a/Engine/Plugins/Experimental/UAF/UAFLayering/Source/UAFLayering/Private/Traits/CachePoseTrait.h b/Engine/Plugins/Experimental/UAF/UAFLayering/Source/UAFLayering/Private/Traits/CachePoseTrait.h
new file mode 100644
--- /dev/null
+++ b/Engine/Plugins/Experimental/UAF/UAFLayering/Source/UAFLayering/Private/Traits/CachePoseTrait.h
@@ -0,0 +1,28 @@
+// Copyright Epic Games, Inc. All Rights Reserved.
+
+#pragma once
+
+#include "Traits/CachePoseTraitData.h"
+#include "TraitCore/Trait.h"
+#include "TraitInterfaces/IEvaluate.h"
+
+namespace UE::UAF
+{
+struct FCachePoseTrait : FAdditiveTrait, IEvaluate
+{
+	DECLARE_ANIM_TRAIT(FCachePoseTrait, FAdditiveTrait)
+
+	struct FInstanceData : FTrait::FInstanceData
+	{
+	};
+
+	using FSharedData = FUAFCachePoseTraitSharedData;
+
+	// IEvaluate impl
+	virtual void PostEvaluate(FEvaluateTraversalContext& Context, const TTraitBinding<IEvaluate>& Binding) const override;
+	
+#if WITH_EDITOR
+	virtual bool IsHidden() const override { return true; };
+#endif
+};
+}
diff --git a/Engine/Plugins/Experimental/UAF/UAFLayering/Source/UAFLayering/Private/Traits/LayerAssetDataProviderTrait.cpp b/Engine/Plugins/Experimental/UAF/UAFLayering/Source/UAFLayering/Private/Traits/LayerAssetDataProviderTrait.cpp
new file mode 100644
--- /dev/null
+++ b/Engine/Plugins/Experimental/UAF/UAFLayering/Source/UAFLayering/Private/Traits/LayerAssetDataProviderTrait.cpp
@@ -0,0 +1,83 @@
+// Copyright Epic Games, Inc. All Rights Reserved.
+
+#include "LayerAssetDataProviderTrait.h"
+
+#include "UAFLogging.h"
+#include "Factory/AnimGraphFactory.h"
+#include "TraitInterfaces/IBlendStack.h"
+#include "TraitInterfaces/IGraphFactory.h"
+
+namespace UE::UAF
+{
+	
+	AUTO_REGISTER_ANIM_TRAIT(FLayerAssetDataProviderTrait)
+
+	#define TRAIT_INTERFACE_ENUMERATOR(GeneratorMacro) \
+		GeneratorMacro(IUpdate) \
+		GeneratorMacro(IGarbageCollection) \
+	
+	GENERATE_ANIM_TRAIT_IMPLEMENTATION(FLayerAssetDataProviderTrait, TRAIT_INTERFACE_ENUMERATOR, NULL_ANIM_TRAIT_INTERFACE_ENUMERATOR, NULL_ANIM_TRAIT_EVENT_ENUMERATOR)
+	#undef TRAIT_INTERFACE_ENUMERATOR
+	#undef TRAIT_EVENT_ENUMERATOR
+
+
+void FLayerAssetDataProviderTrait::FInstanceData::Construct(const FExecutionContext& Context, const FTraitBinding& Binding)
+{
+	FTrait::FInstanceData::Construct(Context, Binding);
+}
+
+void FLayerAssetDataProviderTrait::PreUpdate(FUpdateTraversalContext& Context, const TTraitBinding<IUpdate>& Binding, const FTraitUpdateState& TraitState) const
+{
+	// Check if a new asset has been pushed 
+	const FSharedData* SharedData = Binding.GetSharedData<FSharedData>();
+	FInstanceData* InstanceData = Binding.GetInstanceData<FInstanceData>();
+		
+	const TObjectPtr<UObject> SharedDataAsset = SharedData->GetAsset(Binding);
+	if (SharedDataAsset != InstanceData->CachedAsset)
+	{
+		TTraitBinding<IBlendStack> BlendStackBinding;
+		if(!Binding.GetStackInterface<IBlendStack>(BlendStackBinding))
+		{
+			UAF_TRAIT_LOG(Error, TEXT("Blend Stack Trait is required but not found in stack!"));
+			return;
+		}
+		
+		if (!SharedData->Parameters.IsValid())
+		{
+			UAF_TRAIT_LOG(Warning, TEXT("No valid factory params provided. Can not create graph."));
+			return;
+		}
+		
+		FAnimNextFactoryParams Params = SharedData->Parameters;
+		const UUAFAnimGraph* AnimationGraph = IGraphFactory::GetOrBuildGraph(Context, Binding, SharedDataAsset, Params);
+		if(AnimationGraph == nullptr)
+		{
+			UAF_TRAIT_LOG(Error, TEXT("Failed to create graph!"));
+			return;
+		}
+
+		IBlendStack::FGraphRequest NewGraphRequest;
+		NewGraphRequest.BlendArgs = FAlphaBlendArgs(0.0f);
+		NewGraphRequest.FactoryObject = SharedDataAsset;
+		NewGraphRequest.AnimationGraph = AnimationGraph;
+		NewGraphRequest.FactoryParams = SharedData->Parameters;
+
+		BlendStackBinding.PushGraph(Context, MoveTemp(NewGraphRequest));
+		InstanceData->CachedAsset = SharedDataAsset;
+	}
+		
+	IUpdate::PreUpdate(Context, Binding, TraitState);
+}
+
+void FLayerAssetDataProviderTrait::AddReferencedObjects(const FExecutionContext& Context, const TTraitBinding<IGarbageCollection>& Binding, FReferenceCollector& Collector) const
+{
+	IGarbageCollection::AddReferencedObjects(Context, Binding, Collector);
+		
+	FInstanceData* InstanceData = Binding.GetInstanceData<FInstanceData>();
+	if(InstanceData->CachedAsset != nullptr)
+	{
+		Collector.AddReferencedObject(InstanceData->CachedAsset);
+	}
+}
+	
+}
diff --git a/Engine/Plugins/Experimental/UAF/UAFLayering/Source/UAFLayering/Private/Traits/LayerAssetDataProviderTrait.h b/Engine/Plugins/Experimental/UAF/UAFLayering/Source/UAFLayering/Private/Traits/LayerAssetDataProviderTrait.h
new file mode 100644
--- /dev/null
+++ b/Engine/Plugins/Experimental/UAF/UAFLayering/Source/UAFLayering/Private/Traits/LayerAssetDataProviderTrait.h
@@ -0,0 +1,37 @@
+// Copyright Epic Games, Inc. All Rights Reserved.
+
+#pragma once
+
+#include "TraitCore/Trait.h"
+#include "TraitInterfaces/IGarbageCollection.h"
+#include "TraitInterfaces/IUpdate.h"
+#include "Traits/LayerAssetDataTraitData.h"
+
+namespace UE::UAF
+{
+struct FLayerAssetDataProviderTrait : FAdditiveTrait, IUpdate, IGarbageCollection
+{
+
+	DECLARE_ANIM_TRAIT(FLayerAssetDataProviderTrait, FAdditiveTrait)
+
+	struct FInstanceData : FTrait::FInstanceData
+	{
+		TObjectPtr<UObject> CachedAsset = nullptr;
+		void Construct(const FExecutionContext& Context, const FTraitBinding& Binding);
+	};
+	
+	using FSharedData = FUAFLayerAssetDataTraitSharedData;
+	
+	// IUpdate impl
+	virtual void PreUpdate(FUpdateTraversalContext& Context, const TTraitBinding<IUpdate>& Binding, const FTraitUpdateState& TraitState) const override;
+	
+	// IGarbageCollection impl
+	virtual void AddReferencedObjects(const FExecutionContext& Context, const TTraitBinding<IGarbageCollection>& Binding, FReferenceCollector& Collector) const override;
+	
+#if WITH_EDITOR
+	virtual bool IsHidden() const override { return true; };
+#endif
+	
+};
+	
+}
diff --git a/Engine/Plugins/Experimental/UAF/UAFLayering/Source/UAFLayering/Private/Traits/LayerDataProviderTrait.cpp b/Engine/Plugins/Experimental/UAF/UAFLayering/Source/UAFLayering/Private/Traits/LayerDataProviderTrait.cpp
new file mode 100644
--- /dev/null
+++ b/Engine/Plugins/Experimental/UAF/UAFLayering/Source/UAFLayering/Private/Traits/LayerDataProviderTrait.cpp
@@ -0,0 +1,217 @@
+// Copyright Epic Games, Inc. All Rights Reserved.
+
+#include "Traits/LayerDataProviderTrait.h"
+
+#include "UAFLayerStack.h"
+#include "UAFLogging.h"
+
+namespace UE::UAF
+{
+	AUTO_REGISTER_ANIM_TRAIT(FLayerDataProviderTrait)
+
+	#define TRAIT_INTERFACE_ENUMERATOR(GeneratorMacro) \
+	GeneratorMacro(IContinuousBlend) \
+	GeneratorMacro(IUpdate) \
+	GeneratorMacro(IHierarchy) \
+
+	#define TRAIT_EVENT_ENUMERATOR(GeneratorMacro) \
+	GeneratorMacro(FLayerDataProviderTrait::OnLayerEvent) \
+
+	GENERATE_ANIM_TRAIT_IMPLEMENTATION(FLayerDataProviderTrait, TRAIT_INTERFACE_ENUMERATOR, NULL_ANIM_TRAIT_INTERFACE_ENUMERATOR, TRAIT_EVENT_ENUMERATOR)
+	#undef TRAIT_INTERFACE_ENUMERATOR
+	#undef TRAIT_EVENT_ENUMERATOR
+	
+	void FLayerDataProviderTrait::FInstanceData::Construct(const FExecutionContext& Context, const FTraitBinding& Binding)
+	{
+		const FSharedData* SharedData = Binding.GetSharedData<FSharedData>();
+		LayerProperties = SharedData->GetLayerProperties(Binding);
+		
+		bLayerEnabled = LayerProperties.bLayerEnabled;
+		TargetLayerWeight = LayerProperties.DesiredLayerWeight;
+		EffectiveLayerWeight = LayerProperties.DesiredLayerWeight;
+		LayerBlendInTime = LayerProperties.BlendInTime;
+		LayerBlendOutTime = LayerProperties.BlendOutTime;
+		
+		if (LayerProperties.BlendCurve)
+		{
+			AlphaBlend.SetBlendOption(EAlphaBlendOption::Custom);
+			AlphaBlend.SetCustomCurve(LayerProperties.BlendCurve);
+		}
+		else
+		{
+			AlphaBlend.SetBlendOption(LayerProperties.BlendOption);
+		}
+
+		if (bLayerEnabled)
+		{
+			AlphaBlend.SetBlendTime(LayerBlendInTime); 
+			AlphaBlend.SetValueRange(0.0f, TargetLayerWeight);
+		}
+		else
+		{
+			AlphaBlend.SetBlendTime(LayerBlendOutTime); 
+			AlphaBlend.SetValueRange(TargetLayerWeight, 0.0f);
+		}
+		
+		AlphaBlend.SetAlpha(1.0f);
+		
+		if (SharedData->bCreateCacheInput)
+		{
+			FInstanceData* InstanceData = Binding.GetInstanceData<FInstanceData>();
+			if (!InstanceData->CacheOnlyInput.IsValid())
+			{
+				InstanceData->CacheOnlyInput = Context.AllocateNodeInstance(Binding, SharedData->CacheOnlyInput);
+			}
+		}
+	}
+	
+	void FLayerDataProviderTrait::PostUpdate(FUpdateTraversalContext& Context, const TTraitBinding<IUpdate>& Binding, const FTraitUpdateState& TraitState) const
+	{
+		FInstanceData* InstanceData = Binding.GetInstanceData<FInstanceData>();
+		
+		// This currently mimics executing bindings, once bindings are implemented where this data lives is likely going to change
+		// We only have to ensure that bindings will get executed after any child traits have finished updating, but before they are evaluating - so likely in PostUpdate or PreEval 
+		FUAFLayerProperties& LayerData = InstanceData->LayerProperties;
+		LayerData.ProcessEvents();
+		
+		// Once Bindings/Event have been processed and the desired data has been updated, we can now check if any of them have changed and handle it accordingly
+		// (e.g. override existing data or blend them) 
+		{
+			// Update desired layer weight if it has changed 
+			UpdateTargetLayerWeight(InstanceData, LayerData);
+		
+			// Update Blend In/Out times if they have changed 
+			UpdateLayerBlendTimes(InstanceData, LayerData);
+		
+			// Check if the layer enabled flag has changed and handle it if so 
+			HandleLayerActivation(InstanceData, LayerData);
+		}
+		
+		// After any relevant blend data has been updated we can update the blend itself 
+		InstanceData->AlphaBlend.Update(TraitState.GetDeltaTime());
+		
+		// Update the cached effective layer weight post blend update 
+		InstanceData->EffectiveLayerWeight = InstanceData->AlphaBlend.GetBlendedValue();
+		
+		IUpdate::PostUpdate(Context, Binding, TraitState);
+	}
+	
+	void FLayerDataProviderTrait::UpdateTargetLayerWeight(FInstanceData* InstanceData, const FUAFLayerProperties& LayerData) const
+	{
+		if (LayerData.DesiredLayerWeight != InstanceData->TargetLayerWeight)
+		{
+			InstanceData->TargetLayerWeight = LayerData.DesiredLayerWeight;
+			
+			// Update the layer weight blend to take into account the changed layer weight 
+			if (InstanceData->AlphaBlend.IsComplete())
+			{
+				// If we are fully blended in with no active blend, set the weight directly 
+				if (InstanceData->AlphaBlend.GetDesiredValue() > 0.0f)
+				{
+					InstanceData->AlphaBlend.SetDesiredValue(InstanceData->TargetLayerWeight);
+					InstanceData->AlphaBlend.SetAlpha(1.0f);
+				}
+			}
+			else
+			{
+				// If we are currently blending in, we have to adjust the value range 
+				if (InstanceData->AlphaBlend.GetDesiredValue() > 0.0f)
+				{
+					InstanceData->AlphaBlend.SetValueRange(0.0f, InstanceData->TargetLayerWeight);
+				}
+				else
+				{
+					InstanceData->AlphaBlend.SetValueRange(InstanceData->TargetLayerWeight, 0.0f);
+				}
+			}
+		}
+	}
+	
+	void FLayerDataProviderTrait::UpdateLayerBlendTimes(FInstanceData* InstanceData, const FUAFLayerProperties& LayerData) const
+	{
+		// Update Layer Blend In Time 
+		if (LayerData.BlendInTime != InstanceData->LayerBlendInTime)
+		{
+			InstanceData->LayerBlendInTime = LayerData.BlendInTime;
+			
+			// If we are currently blending in, we have to update the blend time
+			if (!InstanceData->AlphaBlend.IsComplete() && InstanceData->AlphaBlend.GetDesiredValue() > 0.0f)
+			{
+				InstanceData->AlphaBlend.SetBlendTime(InstanceData->LayerBlendInTime);
+			}
+		}
+
+		// Update Layer Blend Out Time
+		if (LayerData.BlendOutTime != InstanceData->LayerBlendOutTime)
+		{
+			InstanceData->LayerBlendOutTime = LayerData.BlendOutTime;
+			if (!InstanceData->AlphaBlend.IsComplete() && InstanceData->AlphaBlend.GetDesiredValue() <= 0.0f)
+			{
+				InstanceData->AlphaBlend.SetBlendTime(InstanceData->LayerBlendOutTime);
+			}
+		}
+	}
+	
+	void FLayerDataProviderTrait::HandleLayerActivation(FInstanceData* InstanceData, const FUAFLayerProperties& LayerData) const
+	{
+		if (LayerData.bLayerEnabled != InstanceData->bLayerEnabled)
+		{
+			InstanceData->bLayerEnabled = LayerData.bLayerEnabled;
+			
+			// If the layer is now enabled, we are blending in. 
+			const bool bBlendingIn = InstanceData->bLayerEnabled;
+			InstanceData->AlphaBlend.SetBlendTime(bBlendingIn ? InstanceData->LayerBlendInTime : InstanceData->LayerBlendOutTime); 
+			InstanceData->AlphaBlend.SetValueRange(InstanceData->EffectiveLayerWeight, bBlendingIn ? LayerData.DesiredLayerWeight : 0.0f);
+		}
+	}
+
+	float FLayerDataProviderTrait::GetBlendWeight(const FExecutionContext& Context, const TTraitBinding<IContinuousBlend>& Binding, int32 ChildIndex) const
+	{		
+		const FInstanceData* InstanceData = Binding.GetInstanceData<FInstanceData>();
+		if (ChildIndex == 0)
+		{
+			return 1.0f - InstanceData->EffectiveLayerWeight;
+		}
+		
+		if (ChildIndex == 1)
+		{
+			return InstanceData->EffectiveLayerWeight;
+		}
+		
+		UAF_TRAIT_LOG(Warning, TEXT("Requested blend weight for invalid child index %d."), ChildIndex);
+		return -1.0f; // Invalid Child Index 
+	}
+
+	ETraitStackPropagation FLayerDataProviderTrait::OnLayerEvent(const FExecutionContext& Context, FTraitBinding& Binding, Layering::FLayerStack_LayerEvent& Event) const
+	{
+		FInstanceData* InstanceData = Binding.GetInstanceData<FInstanceData>();
+		const FSharedData* SharedData = Binding.GetSharedData<FSharedData>();
+		
+		const FName LayerName = SharedData->LayerName;
+		const int32 LayerIndex = SharedData->LayerIndex;
+		const FSoftObjectPath& LayerStackPath = SharedData->LayerStackPath;
+		if ( (LayerName == Event.LayerName || LayerIndex == Event.LayerIndex) && LayerStackPath == Event.LayerStackAssetPath)
+		{
+			InstanceData->LayerProperties.AddEvent(Event);
+			Event.bAutoConsumeEvent ? Event.MarkConsumed() : Event.MarkHandled();
+		}
+		
+		return ETraitStackPropagation::Continue;
+	}
+
+	uint32 FLayerDataProviderTrait::GetNumChildren(const FExecutionContext& Context, const TTraitBinding<IHierarchy>& Binding) const
+	{
+		const FInstanceData* InstanceData = Binding.GetInstanceData<FInstanceData>();
+		return InstanceData->CacheOnlyInput.IsValid() ? 1 : 0;
+	}
+
+	void FLayerDataProviderTrait::GetChildren(const FExecutionContext& Context, const TTraitBinding<IHierarchy>& Binding, FChildrenArray& Children) const
+	{
+		// Only do this if cache only mode 
+		const FInstanceData* InstanceData = Binding.GetInstanceData<FInstanceData>();
+		if (InstanceData->CacheOnlyInput.IsValid())
+		{
+			Children.Add(InstanceData->CacheOnlyInput);
+		}
+	}
+}
diff --git a/Engine/Plugins/Experimental/UAF/UAFLayering/Source/UAFLayering/Private/Traits/LayerDataProviderTrait.h b/Engine/Plugins/Experimental/UAF/UAFLayering/Source/UAFLayering/Private/Traits/LayerDataProviderTrait.h
new file mode 100644
--- /dev/null
+++ b/Engine/Plugins/Experimental/UAF/UAFLayering/Source/UAFLayering/Private/Traits/LayerDataProviderTrait.h
@@ -0,0 +1,74 @@
+// Copyright Epic Games, Inc. All Rights Reserved.
+
+#pragma once
+
+#include "Traits/LayerDataProviderTraitData.h"
+#include "TraitCore/Trait.h"
+#include "TraitInterfaces/IContinuousBlend.h"
+#include "TraitInterfaces/IHierarchy.h"
+#include "TraitInterfaces/IUpdate.h"
+
+namespace UE::UAF
+{
+struct FLayerDataProviderTrait : FAdditiveTrait, IUpdate, IContinuousBlend, IHierarchy
+{
+	DECLARE_ANIM_TRAIT(FLayerDataProviderTrait, FAdditiveTrait)
+
+	struct FInstanceData : FTrait::FInstanceData
+	{
+		// If the layer is currently enabled or not. 
+		bool bLayerEnabled = true;
+		
+		// The cached effective layer weight - the actual weight used when applying this layer based on the target weight and any active blends
+		float EffectiveLayerWeight = 1.0f;
+		
+		// The target weight of this layer. This is the desired weight to apply for this layer, however if the layer is currently blending in/out 
+		// this might not be the weight that gets applied - see EffectiveLayerWeight. 
+		float TargetLayerWeight = 1.0f;
+		
+		// The blend in time used when blending this layer in
+		float LayerBlendInTime = 0.0f;
+		
+		// The blend out time used when blending this layer out 
+		float LayerBlendOutTime = 0.0f;
+		
+		// Trait handle to execute cache only parts of the graph
+		FTraitPtr CacheOnlyInput;
+		
+		// This updated the discrete blend when a layer blends in or out 
+		FAlphaBlend AlphaBlend;
+		
+		// This layer data is likely not going to live here once bindings are implemented
+		// This contains all the bound / desired layer properties versus the rest of the instance data containing the currently used data based on
+		// any blending between properties and to check if data has changed between frames 
+		FUAFLayerProperties LayerProperties;
+		
+		void Construct(const FExecutionContext& Context, const FTraitBinding& Binding);
+	};
+
+	using FSharedData = FUAFLayerDataProviderTraitSharedData;
+	
+	// IUpdate imp
+	virtual void PostUpdate(FUpdateTraversalContext& Context, const TTraitBinding<IUpdate>& Binding, const FTraitUpdateState& TraitState) const override;
+	
+	// IContinuousBlend impl - this allows us to override the weight for any lower blend trait 
+	virtual float GetBlendWeight(const FExecutionContext& Context, const TTraitBinding<IContinuousBlend>& Binding, int32 ChildIndex) const override;
+
+	// Event Handling
+	ETraitStackPropagation OnLayerEvent(const FExecutionContext& Context, FTraitBinding& Binding, Layering::FLayerStack_LayerEvent& Event) const;
+	
+	// IHierarchy Impl
+	virtual uint32 GetNumChildren(const FExecutionContext& Context, const TTraitBinding<IHierarchy>& Binding) const override;
+	virtual void GetChildren(const FExecutionContext& Context, const TTraitBinding<IHierarchy>& Binding, FChildrenArray& Children) const override;
+	
+#if WITH_EDITOR
+	virtual bool IsHidden() const override { return true; };
+#endif
+	
+private:
+	void UpdateTargetLayerWeight(FInstanceData* InstanceData, const FUAFLayerProperties& LayerData) const;
+	void UpdateLayerBlendTimes(FInstanceData* InstanceData, const FUAFLayerProperties& LayerData) const;
+	void HandleLayerActivation(FInstanceData* InstanceData, const FUAFLayerProperties& LayerData) const;
+	
+};
+}
diff --git a/Engine/Plugins/Experimental/UAF/UAFLayering/Source/UAFLayering/Private/Traits/LayerDataProviderTraitData.cpp b/Engine/Plugins/Experimental/UAF/UAFLayering/Source/UAFLayering/Private/Traits/LayerDataProviderTraitData.cpp
new file mode 100644
--- /dev/null
+++ b/Engine/Plugins/Experimental/UAF/UAFLayering/Source/UAFLayering/Private/Traits/LayerDataProviderTraitData.cpp
@@ -0,0 +1,45 @@
+// Copyright Epic Games, Inc. All Rights Reserved.
+
+#include "Traits/LayerDataProviderTraitData.h"
+
+void FUAFLayerProperties::AddEvent(const UE::UAF::Layering::FLayerStack_LayerEvent& Event)
+{
+	EventsToProcess.Add(Event);
+}
+
+void FUAFLayerProperties::ProcessEvents()
+{
+	using namespace UE::UAF::Layering;
+	
+	for (const FLayerStack_LayerEvent& Event : EventsToProcess)
+	{
+		switch (Event.Action)
+		{
+		case ELayerEventAction::EnableLayer:
+			bLayerEnabled = true;
+			break;
+		case ELayerEventAction::DisableLayer:
+			bLayerEnabled = false;
+			break;
+		case ELayerEventAction::SetFloatValue:
+			if (Event.PropertyToSet.IsEqual(FLayerStack_LayerEvent::LayerWeightProperty))
+			{
+				DesiredLayerWeight = Event.FloatValue;
+			}
+			else if (Event.PropertyToSet.IsEqual(FLayerStack_LayerEvent::BlendInTimeProperty))
+			{
+				BlendInTime = Event.FloatValue;
+			}
+			else if (Event.PropertyToSet.IsEqual(FLayerStack_LayerEvent::BlendOutTimeProperty))
+			{
+				BlendOutTime = Event.FloatValue;
+			}
+			break;
+		case ELayerEventAction::None:
+		default:
+			break;
+		}
+	}
+	
+	EventsToProcess.Empty();
+}
diff --git a/Engine/Plugins/Experimental/UAF/UAFLayering/Source/UAFLayering/Private/UAFLayeringModule.cpp b/Engine/Plugins/Experimental/UAF/UAFLayering/Source/UAFLayering/Private/UAFLayeringModule.cpp
new file mode 100644
--- /dev/null
+++ b/Engine/Plugins/Experimental/UAF/UAFLayering/Source/UAFLayering/Private/UAFLayeringModule.cpp
@@ -0,0 +1,29 @@
+// Copyright Epic Games, Inc. All Rights Reserved.
+
+#include "UAFLayeringModule.h"
+
+#include "Factory/AnimNextFactoryParams.h"
+#include "RigVMCore/RigVMRegistry.h"
+#include "Traits/LayerDataProviderTraitData.h"
+
+#define LOCTEXT_NAMESPACE "FUAFLayeringModule"
+
+void FUAFLayeringModule::StartupModule()
+{
+	static UScriptStruct* const AllowedStructTypes[] =
+	{
+		FUAFLayerProperties::StaticStruct(),
+		FAnimNextFactoryParams::StaticStruct(),
+	};
+	
+	FRigVMRegistry::Get().RegisterStructTypes(AllowedStructTypes);
+}
+
+void FUAFLayeringModule::ShutdownModule()
+{
+	
+}
+
+#undef LOCTEXT_NAMESPACE
+	
+IMPLEMENT_MODULE(FUAFLayeringModule, UAFLayering)
diff --git a/Engine/Plugins/Experimental/UAF/UAFLayering/Source/UAFLayering/Private/UAFLayeringModule.h b/Engine/Plugins/Experimental/UAF/UAFLayering/Source/UAFLayering/Private/UAFLayeringModule.h
new file mode 100644
--- /dev/null
+++ b/Engine/Plugins/Experimental/UAF/UAFLayering/Source/UAFLayering/Private/UAFLayeringModule.h
@@ -0,0 +1,13 @@
+// Copyright Epic Games, Inc. All Rights Reserved.
+
+#pragma once
+
+#include "Modules/ModuleManager.h"
+
+class FUAFLayeringModule : public IModuleInterface
+{
+public:
+	/** IModuleInterface implementation */
+	virtual void StartupModule() override;
+	virtual void ShutdownModule() override;
+};
diff --git a/Engine/Plugins/Experimental/UAF/UAFLayering/Source/UAFLayering/Private/UAFLayeringTypes.cpp b/Engine/Plugins/Experimental/UAF/UAFLayering/Source/UAFLayering/Private/UAFLayeringTypes.cpp
new file mode 100644
--- /dev/null
+++ b/Engine/Plugins/Experimental/UAF/UAFLayering/Source/UAFLayering/Private/UAFLayeringTypes.cpp
@@ -0,0 +1,11 @@
+// Copyright Epic Games, Inc. All Rights Reserved.
+
+#include "UAFLayeringTypes.h"
+
+namespace UE::UAF::Layering
+{
+	const FName FLayerStack_LayerEvent::LayerWeightProperty(TEXT("LayerWeight"));
+	const FName FLayerStack_LayerEvent::BlendInTimeProperty(TEXT("BlendInTime"));
+	const FName FLayerStack_LayerEvent::BlendOutTimeProperty(TEXT("BlendOutTime"));
+	
+}
diff --git a/Engine/Plugins/Experimental/UAF/UAFLayering/Source/UAFLayering/Private/UAFLayeringUtils.cpp b/Engine/Plugins/Experimental/UAF/UAFLayering/Source/UAFLayering/Private/UAFLayeringUtils.cpp
new file mode 100644
--- /dev/null
+++ b/Engine/Plugins/Experimental/UAF/UAFLayering/Source/UAFLayering/Private/UAFLayeringUtils.cpp
@@ -0,0 +1,83 @@
+// Copyright Epic Games, Inc. All Rights Reserved.
+
+#include "UAFLayeringUtils.h"
+
+#include "UAFLayeringTypes.h"
+#include "Component/AnimNextComponent.h"
+#include "TraitCore/TraitEvent.h"
+
+void UUAFLayeringUtils::DisableLayer(UAnimNextComponent* UAFComponent, FName LayerName, const TSoftObjectPtr<UUAFLayerStack> LayerStackPath)
+{
+	if (!UAFComponent || LayerName == NAME_None || LayerStackPath.IsNull())
+	{
+		return;
+	}
+	
+	TSharedPtr<UE::UAF::Layering::FLayerStack_LayerEvent> Event = MakeTraitEvent<UE::UAF::Layering::FLayerStack_LayerEvent>(LayerName, LayerStackPath.ToSoftObjectPath());
+	Event->Action = UE::UAF::Layering::ELayerEventAction::DisableLayer;
+	UAFComponent->QueueInputTraitEvent(Event);
+}
+
+void UUAFLayeringUtils::EnableLayer(UAnimNextComponent* UAFComponent, FName LayerName, const TSoftObjectPtr<UUAFLayerStack> LayerStackPath)
+{
+	if (!UAFComponent || LayerName == NAME_None || LayerStackPath.IsNull())
+	{
+		return;
+	}
+	
+	TSharedPtr<UE::UAF::Layering::FLayerStack_LayerEvent> Event = MakeTraitEvent<UE::UAF::Layering::FLayerStack_LayerEvent>(LayerName, LayerStackPath.ToSoftObjectPath());
+	Event->Action = UE::UAF::Layering::ELayerEventAction::EnableLayer;
+	UAFComponent->QueueInputTraitEvent(Event);
+}
+
+void UUAFLayeringUtils::SetLayerWeight(UAnimNextComponent* UAFComponent, FName LayerName, const TSoftObjectPtr<UUAFLayerStack> LayerStackPath, const float Weight)
+{
+	if (!UAFComponent || LayerName == NAME_None || LayerStackPath.IsNull())
+	{
+		return;
+	}
+	
+	TSharedPtr<UE::UAF::Layering::FLayerStack_LayerEvent> Event = MakeTraitEvent<UE::UAF::Layering::FLayerStack_LayerEvent>(LayerName, LayerStackPath.ToSoftObjectPath());
+	Event->Action = UE::UAF::Layering::ELayerEventAction::SetFloatValue;
+	Event->PropertyToSet = UE::UAF::Layering::FLayerStack_LayerEvent::LayerWeightProperty;
+	Event->FloatValue = Weight;
+	UAFComponent->QueueInputTraitEvent(Event);
+}
+
+void UUAFLayeringUtils::DisableLayerByIndex(UAnimNextComponent* UAFComponent, int32 LayerIndex, const TSoftObjectPtr<UUAFLayerStack> LayerStackPath)
+{
+	if (!UAFComponent || LayerIndex == INDEX_NONE || LayerStackPath.IsNull())
+	{
+		return;
+	}
+	
+	TSharedPtr<UE::UAF::Layering::FLayerStack_LayerEvent> Event = MakeTraitEvent<UE::UAF::Layering::FLayerStack_LayerEvent>(LayerIndex, LayerStackPath.ToSoftObjectPath());
+	Event->Action = UE::UAF::Layering::ELayerEventAction::DisableLayer;
+	UAFComponent->QueueInputTraitEvent(Event);
+}
+
+void UUAFLayeringUtils::EnableLayerByIndex(UAnimNextComponent* UAFComponent, int32 LayerIndex, const TSoftObjectPtr<UUAFLayerStack> LayerStackPath)
+{
+	if (!UAFComponent || LayerIndex == INDEX_NONE || LayerStackPath.IsNull())
+	{
+		return;
+	}
+	
+	TSharedPtr<UE::UAF::Layering::FLayerStack_LayerEvent> Event = MakeTraitEvent<UE::UAF::Layering::FLayerStack_LayerEvent>(LayerIndex, LayerStackPath.ToSoftObjectPath());
+	Event->Action = UE::UAF::Layering::ELayerEventAction::EnableLayer;
+	UAFComponent->QueueInputTraitEvent(Event);
+}
+
+void UUAFLayeringUtils::SetLayerWeightByIndex(UAnimNextComponent* UAFComponent, int32 LayerIndex, const TSoftObjectPtr<UUAFLayerStack> LayerStackPath, const float Weight)
+{
+	if (!UAFComponent || LayerIndex == INDEX_NONE || LayerStackPath.IsNull())
+	{
+		return;
+	}
+	
+	TSharedPtr<UE::UAF::Layering::FLayerStack_LayerEvent> Event = MakeTraitEvent<UE::UAF::Layering::FLayerStack_LayerEvent>(LayerIndex, LayerStackPath.ToSoftObjectPath());
+	Event->Action = UE::UAF::Layering::ELayerEventAction::SetFloatValue;
+	Event->PropertyToSet =  UE::UAF::Layering::FLayerStack_LayerEvent::LayerWeightProperty;
+	Event->FloatValue = Weight;
+	UAFComponent->QueueInputTraitEvent(Event);
+}
diff --git a/Engine/Plugins/Experimental/UAF/UAFLayering/Source/UAFLayering/Public/UAFLayeringTypes.h b/Engine/Plugins/Experimental/UAF/UAFLayering/Source/UAFLayering/Public/UAFLayeringTypes.h
new file mode 100644
--- /dev/null
+++ b/Engine/Plugins/Experimental/UAF/UAFLayering/Source/UAFLayering/Public/UAFLayeringTypes.h
@@ -0,0 +1,59 @@
+// Copyright Epic Games, Inc. All Rights Reserved.
+
+#pragma once
+
+#include "TraitCore/TraitEvent.h"
+
+namespace UE::UAF::Layering
+{
+	
+enum class ELayerEventAction : uint8
+{
+	None,
+	EnableLayer,
+	DisableLayer,
+	SetFloatValue
+};
+	
+struct FLayerStack_LayerEvent : public FAnimNextTraitEvent
+{
+	DECLARE_ANIM_TRAIT_EVENT(FLayerStack_LayerEvent, FAnimNextTraitEvent)
+
+	FLayerStack_LayerEvent(FName InLayerName, const FSoftObjectPath& InObjectPath)
+		: LayerName(InLayerName)
+		, LayerStackAssetPath(InObjectPath)
+	{}
+	
+	FLayerStack_LayerEvent(int32 InLayerIndex, const FSoftObjectPath& InObjectPath)
+		: LayerIndex(InLayerIndex)
+		, LayerStackAssetPath(InObjectPath)
+	{}
+	
+	// The name of the layer to affect 
+	const FName LayerName = NAME_None;
+	
+	// The index in the layer stack for the layer to affect
+	const int32 LayerIndex = INDEX_NONE;
+
+	// The soft object path of the layer stack asset the layer exists in 
+	const FSoftObjectPath LayerStackAssetPath = FSoftObjectPath();
+
+	// The action to perform 
+	ELayerEventAction Action = ELayerEventAction::None;
+
+	// The name of the property to change if applicable 
+	FName PropertyToSet = NAME_None;
+
+	// The float value to set, if applicable
+	float FloatValue = 0.0f;
+	
+	// If this event should be consumed by the first found layer that matches the description
+	// Per default the event will be passed through the UAF system and every layer gets the option to react to it 
+	bool bAutoConsumeEvent = false;
+	
+	static const FName LayerWeightProperty;	
+	static const FName BlendInTimeProperty;	
+	static const FName BlendOutTimeProperty;
+};
+	
+}
diff --git a/Engine/Plugins/Experimental/UAF/UAFLayering/Source/UAFLayering/Public/UAFLayeringUtils.h b/Engine/Plugins/Experimental/UAF/UAFLayering/Source/UAFLayering/Public/UAFLayeringUtils.h
new file mode 100644
--- /dev/null
+++ b/Engine/Plugins/Experimental/UAF/UAFLayering/Source/UAFLayering/Public/UAFLayeringUtils.h
@@ -0,0 +1,44 @@
+// Copyright Epic Games, Inc. All Rights Reserved.
+
+#pragma once
+
+#include "CoreMinimal.h"
+#include "Kismet/BlueprintFunctionLibrary.h"
+#include "UAFLayeringUtils.generated.h"
+
+class UUAFLayerStack;
+class UAnimNextComponent; 
+
+#define UE_API UAFLAYERING_API
+
+UCLASS(MinimalAPI)
+class UUAFLayeringUtils : public UBlueprintFunctionLibrary
+{
+	GENERATED_BODY()
+	
+	// Enables the specified layer by name in the specified Layer Stack
+	UFUNCTION(BlueprintCallable, Category = "UAF|Layering", meta = (ScriptMethod))
+	static UE_API void EnableLayer(UAnimNextComponent* UAFComponent, FName LayerName, const TSoftObjectPtr<UUAFLayerStack> LayerStackPath);
+	
+	// Disables the specified layer by name in the specified Layer Stack
+	UFUNCTION(BlueprintCallable, Category = "UAF|Layering", meta = (ScriptMethod))
+	static UE_API void DisableLayer(UAnimNextComponent* UAFComponent, FName LayerName, const TSoftObjectPtr<UUAFLayerStack> LayerStackPath);
+	
+	// Sets the weight of the specified layer by name in the specified Layer Stack
+	UFUNCTION(BlueprintCallable, Category = "UAF|Layering", meta = (ScriptMethod))
+	static UE_API void SetLayerWeight(UAnimNextComponent* UAFComponent, FName LayerName, const TSoftObjectPtr<UUAFLayerStack> LayerStackPath, const float Weight);
+	
+	// Enables the specified layer by index in the specified Layer Stack
+	UFUNCTION(BlueprintCallable, Category = "UAF|Layering", meta = (ScriptMethod))
+	static UE_API void EnableLayerByIndex(UAnimNextComponent* UAFComponent, int32 LayerIndex, const TSoftObjectPtr<UUAFLayerStack> LayerStackPath);
+	
+	// Disables the specified layer by index in the specified Layer Stack
+	UFUNCTION(BlueprintCallable, Category = "UAF|Layering", meta = (ScriptMethod))
+	static UE_API void DisableLayerByIndex(UAnimNextComponent* UAFComponent, int32 LayerIndex, const TSoftObjectPtr<UUAFLayerStack> LayerStackPath);
+	
+	// Sets the weight of the specified layer by index in the specified Layer Stack
+	UFUNCTION(BlueprintCallable, Category = "UAF|Layering", meta = (ScriptMethod))
+	static UE_API void SetLayerWeightByIndex(UAnimNextComponent* UAFComponent, int32 LayerIndex, const TSoftObjectPtr<UUAFLayerStack> LayerStackPath, const float Weight);
+};
+
+#undef UE_API
diff --git a/Engine/Plugins/Experimental/UAF/UAFLayering/Source/UAFLayering/UAFLayering.Build.cs b/Engine/Plugins/Experimental/UAF/UAFLayering/Source/UAFLayering/UAFLayering.Build.cs
new file mode 100644
--- /dev/null
+++ b/Engine/Plugins/Experimental/UAF/UAFLayering/Source/UAFLayering/UAFLayering.Build.cs
@@ -0,0 +1,43 @@
+// Copyright Epic Games, Inc. All Rights Reserved.
+
+using UnrealBuildTool;
+
+public class UAFLayering : ModuleRules
+{
+	public UAFLayering(ReadOnlyTargetRules Target) : base(Target)
+	{
+		PCHUsage = ModuleRules.PCHUsageMode.UseExplicitOrSharedPCHs;
+		
+		PublicDependencyModuleNames.AddRange(
+			new string[]
+			{
+				"Core",
+				"UAF"
+				// ... add other public dependencies that you statically link with here ...
+			}
+			);
+			
+		
+		PrivateDependencyModuleNames.AddRange(
+			new string[]
+			{
+				"CoreUObject",
+				"Engine",
+				"Slate",
+				"SlateCore",
+				"UAFAnimGraph",
+				"RigVM",
+				"UAFAnimGraph"
+				// ... add private dependencies that you statically link with here ...	
+			}
+			);
+		
+		
+		DynamicallyLoadedModuleNames.AddRange(
+			new string[]
+			{
+				// ... add any modules that your module loads dynamically here ...
+			}
+			);
+	}
+}
diff --git a/Engine/Plugins/Experimental/UAF/UAFLayering/Source/UAFLayeringEditor/Internal/UAFLayerStackFactory.h b/Engine/Plugins/Experimental/UAF/UAFLayering/Source/UAFLayeringEditor/Internal/UAFLayerStackFactory.h
new file mode 100644
--- /dev/null
+++ b/Engine/Plugins/Experimental/UAF/UAFLayering/Source/UAFLayeringEditor/Internal/UAFLayerStackFactory.h
@@ -0,0 +1,19 @@
+// Copyright Epic Games, Inc. All Rights Reserved.
+
+#pragma once
+
+#include "CoreMinimal.h"
+#include "Factories/Factory.h"
+#include "UAFLayerStackFactory.generated.h"
+
+UCLASS(MinimalAPI)
+class UUAFLayerStackFactory : public UFactory
+{
+	GENERATED_BODY()
+
+	UUAFLayerStackFactory();
+
+	virtual bool ConfigureProperties() override;
+	virtual UObject* FactoryCreateNew(UClass* Class, UObject* InParent, FName Name, EObjectFlags Flags, UObject* Context, FFeedbackContext* Warn, FName CallingContext) override;
+	virtual UObject* FactoryCreateNew(UClass* Class, UObject* InParent, FName Name, EObjectFlags Flags, UObject* Context, FFeedbackContext* Warn) override;
+};
\ No newline at end of file
diff --git a/Engine/Plugins/Experimental/UAF/UAFLayering/Source/UAFLayeringEditor/Private/UAFLayerStackAssetDefinition.cpp b/Engine/Plugins/Experimental/UAF/UAFLayering/Source/UAFLayeringEditor/Private/UAFLayerStackAssetDefinition.cpp
new file mode 100644
--- /dev/null
+++ b/Engine/Plugins/Experimental/UAF/UAFLayering/Source/UAFLayeringEditor/Private/UAFLayerStackAssetDefinition.cpp
@@ -0,0 +1,46 @@
+// Copyright Epic Games, Inc. All Rights Reserved.
+
+#include "UAFLayerStackAssetDefinition.h"
+#include "UAFLayerStack.h"
+#include "IWorkspaceEditorModule.h"
+#include "Modules/ModuleManager.h"
+#include "Workspace/AnimNextWorkspaceFactory.h"
+
+#define LOCTEXT_NAMESPACE "AnimNextAssetDefinitions"
+
+FText UAssetDefinition_UAFLayerStack::GetAssetDisplayName() const
+{
+	return LOCTEXT("UAFLayerStack", "UAF Layer Stack"); 
+}
+
+FLinearColor UAssetDefinition_UAFLayerStack::GetAssetColor() const
+{ 
+	return FLinearColor(FColor(148, 49, 145));
+}
+
+TSoftClassPtr<UObject> UAssetDefinition_UAFLayerStack::GetAssetClass() const
+{ 
+	return UUAFLayerStack::StaticClass();
+}
+
+TConstArrayView<FAssetCategoryPath> UAssetDefinition_UAFLayerStack::GetAssetCategories() const
+{
+	static const auto Categories = { FAssetCategoryPath(EAssetCategoryPaths::Animation, LOCTEXT("UAFSubMenu", "Animation Framework")) };
+	return Categories;
+}
+
+EAssetCommandResult UAssetDefinition_UAFLayerStack::OpenAssets(const FAssetOpenArgs& OpenArgs) const
+{
+	using namespace UE::UAF::Editor;
+	using namespace UE::Workspace;
+
+	for (UUAFLayerStack* Asset : OpenArgs.LoadObjects<UUAFLayerStack>())
+	{
+		IWorkspaceEditorModule& WorkspaceEditorModule = FModuleManager::Get().LoadModuleChecked<IWorkspaceEditorModule>("WorkspaceEditor");
+		WorkspaceEditorModule.OpenWorkspaceForObject(Asset, EOpenWorkspaceMethod::Default, UAnimNextWorkspaceFactory::StaticClass());
+	}
+
+	return EAssetCommandResult::Handled;
+}
+
+#undef LOCTEXT_NAMESPACE
\ No newline at end of file
diff --git a/Engine/Plugins/Experimental/UAF/UAFLayering/Source/UAFLayeringEditor/Private/UAFLayerStackAssetDefinition.h b/Engine/Plugins/Experimental/UAF/UAFLayering/Source/UAFLayeringEditor/Private/UAFLayerStackAssetDefinition.h
new file mode 100644
--- /dev/null
+++ b/Engine/Plugins/Experimental/UAF/UAFLayering/Source/UAFLayeringEditor/Private/UAFLayerStackAssetDefinition.h
@@ -0,0 +1,22 @@
+// Copyright Epic Games, Inc. All Rights Reserved.
+
+#pragma once
+
+#include "AssetDefinitionDefault.h"
+#include "UAFLayerStackAssetDefinition.generated.h"
+
+UCLASS()
+class UAssetDefinition_UAFLayerStack : public UAssetDefinitionDefault
+{
+	GENERATED_BODY()
+
+public:
+	// Begin UAssetDefinition Interface
+	virtual FText GetAssetDisplayName() const override; 
+	virtual FLinearColor GetAssetColor() const override;
+	virtual TSoftClassPtr<UObject> GetAssetClass() const override;
+	virtual TConstArrayView<FAssetCategoryPath> GetAssetCategories() const override;
+	virtual EAssetCommandResult OpenAssets(const FAssetOpenArgs& OpenArgs) const override;
+	// End UAssetDefinition Interface
+};
+
diff --git a/Engine/Plugins/Experimental/UAF/UAFLayering/Source/UAFLayeringEditor/Private/UAFLayerStackFactory.cpp b/Engine/Plugins/Experimental/UAF/UAFLayering/Source/UAFLayeringEditor/Private/UAFLayerStackFactory.cpp
new file mode 100644
--- /dev/null
+++ b/Engine/Plugins/Experimental/UAF/UAFLayering/Source/UAFLayeringEditor/Private/UAFLayerStackFactory.cpp
@@ -0,0 +1,50 @@
+// Copyright Epic Games, Inc. All Rights Reserved.
+
+#include "UAFLayerStackFactory.h"
+#include "UAFLayerStack.h"
+#include "UAFLayerStack_EditorData.h"
+#include "UAFCompilationScope.h"
+
+UUAFLayerStackFactory::UUAFLayerStackFactory()
+{
+	bCreateNew = true;
+	bEditAfterNew = true;
+	SupportedClass = UUAFLayerStack::StaticClass();
+}
+
+bool UUAFLayerStackFactory::ConfigureProperties()
+{
+	return true;
+}
+
+UObject* UUAFLayerStackFactory::FactoryCreateNew(UClass* Class, UObject* InParent, FName Name, EObjectFlags Flags, UObject* Context, FFeedbackContext* Warn, FName CallingContext)
+{
+	EObjectFlags FlagsToUse = Flags | RF_Public | RF_Standalone | RF_Transactional | RF_LoadCompleted;
+	if (InParent == GetTransientPackage())
+	{
+		FlagsToUse &= ~RF_Standalone;
+	}
+
+	// Create LayerStack Asset
+	UUAFLayerStack* LayerStack = NewObject<UUAFLayerStack>(InParent, Class, Name, FlagsToUse);
+
+	// Create related Editor Data
+	const FName EditorDataName = FName(Name.ToString() + "_Layering_EditorData");
+	UUAFLayerStack_EditorData* EditorData = NewObject<UUAFLayerStack_EditorData>(LayerStack, EditorDataName, RF_Transactional);
+
+	// Set the layer controllers editor data reference
+	LayerStack->EditorData = EditorData;
+
+	EditorData->bUsesExternalPackages = false;
+	EditorData->Initialize(/*bRecompileVM*/false);
+
+	UE::UAF::UncookedOnly::Compilation::RequestAssetCompilation(LayerStack);
+	ensure(!EditorData->bErrorsDuringCompilation);
+
+	return LayerStack;
+}
+
+UObject* UUAFLayerStackFactory::FactoryCreateNew(UClass* Class, UObject* InParent, FName Name, EObjectFlags Flags, UObject* Context, FFeedbackContext* Warn)
+{
+	return FactoryCreateNew(Class, InParent, Name, Flags, Context, Warn, NAME_None);
+}
diff --git a/Engine/Plugins/Experimental/UAF/UAFLayering/Source/UAFLayeringEditor/Private/UAFLayeringEditorMode.cpp b/Engine/Plugins/Experimental/UAF/UAFLayering/Source/UAFLayeringEditor/Private/UAFLayeringEditorMode.cpp
new file mode 100644
--- /dev/null
+++ b/Engine/Plugins/Experimental/UAF/UAFLayering/Source/UAFLayeringEditor/Private/UAFLayeringEditorMode.cpp
@@ -0,0 +1,6 @@
+// Copyright Epic Games, Inc. All Rights Reserved.
+
+
+#include "UAFLayeringEditorMode.h"
+
+const FEditorModeID UUAFLayeringEditorMode::EM_UAFLayeringMode("UAFLayeringEditorMode");
\ No newline at end of file
diff --git a/Engine/Plugins/Experimental/UAF/UAFLayering/Source/UAFLayeringEditor/Private/UAFLayeringEditorMode.h b/Engine/Plugins/Experimental/UAF/UAFLayering/Source/UAFLayeringEditor/Private/UAFLayeringEditorMode.h
new file mode 100644
--- /dev/null
+++ b/Engine/Plugins/Experimental/UAF/UAFLayering/Source/UAFLayeringEditor/Private/UAFLayeringEditorMode.h
@@ -0,0 +1,16 @@
+// Copyright Epic Games, Inc. All Rights Reserved.
+
+#pragma once
+
+#include "CoreMinimal.h"
+#include "Tools/UEdMode.h"
+#include "UAFLayeringEditorMode.generated.h"
+
+UCLASS()
+class UUAFLayeringEditorMode : public UEdMode
+{
+	GENERATED_BODY()
+public:
+	const static FEditorModeID EM_UAFLayeringMode;
+	
+};
diff --git a/Engine/Plugins/Experimental/UAF/UAFLayering/Source/UAFLayeringEditor/Private/UAFLayeringEditorModule.cpp b/Engine/Plugins/Experimental/UAF/UAFLayering/Source/UAFLayeringEditor/Private/UAFLayeringEditorModule.cpp
new file mode 100644
--- /dev/null
+++ b/Engine/Plugins/Experimental/UAF/UAFLayering/Source/UAFLayeringEditor/Private/UAFLayeringEditorModule.cpp
@@ -0,0 +1,119 @@
+// Copyright Epic Games, Inc. All Rights Reserved.
+
+#include "UAFLayeringEditorModule.h"
+
+#include "ContextObjectStore.h"
+#include "EditorModeManager.h"
+#include "EditorUtils.h"
+#include "IAnimNextEditorModule.h"
+#include "IWorkspaceEditor.h"
+#include "UAFLayeringEditorMode.h"
+#include "UAFLayeringStyle.h"
+#include "UAFLayerStack.h"
+#include "Layers/UAFLayer.h"
+#include "Tools/EdModeInteractiveToolsContext.h"
+#include "Widgets/SLayerStack.h"
+#include "Workspace/LayerStackItemDetails.h"
+#include "Workspace/UAFLayerStackWorkspaceOutlinerData.h"
+
+
+#define LOCTEXT_NAMESPACE "FUAFLayeringEditorModule"
+
+namespace UE::UAF::LayeringEditor
+{
+void FUAFLayeringEditorModule::StartupModule()
+{
+	FUAFLayeringStyle::Get();
+
+	UE::UAF::Editor::IAnimNextEditorModule& UAFEditorModule = FModuleManager::Get().LoadModuleChecked<UE::UAF::Editor::IAnimNextEditorModule>("UAFEditor");
+	UAFEditorModule.AddWorkspaceSupportedAssetClass(UUAFLayerStack::StaticClass()->GetClassPathName());
+
+	// Build Document Args
+	UE::Workspace::FObjectDocumentArgs LayerStackDocumentArgs;
+	LayerStackDocumentArgs.OnMakeDocumentWidget.BindRaw(this, &FUAFLayeringEditorModule::MakeDocumentWidget);
+	LayerStackDocumentArgs.OnGetDocumentBreadcrumbTrail.BindRaw(this, &FUAFLayeringEditorModule::GetBreadcrumbTrail);
+	LayerStackDocumentArgs.SpawnLocation = UE::Workspace::WorkspaceTabs::TopMiddleDocumentArea;
+	LayerStackDocumentArgs.DocumentEditorMode = UUAFLayeringEditorMode::EM_UAFLayeringMode;
+	LayerStackDocumentArgs.OnGetTabIcon = UE::Workspace::FOnGetTabIcon::CreateLambda([](const UE::Workspace::FWorkspaceEditorContext& InContext)
+		{
+			return FUAFLayeringStyle::Get().GetBrush(TEXT("ClassIcon.UAFLayerStack"));
+		});
+
+	UE::Workspace::IWorkspaceEditorModule& WorkspaceEditorModule = FModuleManager::Get().LoadModuleChecked<UE::Workspace::IWorkspaceEditorModule>("WorkspaceEditor");
+	WorkspaceEditorModule.RegisterObjectDocumentType(UUAFLayerStack::StaticClass()->GetClassPathName(), LayerStackDocumentArgs);
+	WorkspaceEditorModule.RegisterWorkspaceItemDetails(UE::Workspace::FOutlinerItemDetailsId(FUAFLayerStackWorkspaceOutlinerData::StaticStruct()->GetFName()), MakeShared<UE::UAF::LayeringEditor::ULayerStackItemDetails>());
+}
+
+void FUAFLayeringEditorModule::ShutdownModule()
+{
+	if (UObjectInitialized())
+	{
+		// Unregister document type for layer stack class from workspace 
+		if (UE::Workspace::IWorkspaceEditorModule* WorkspaceEditorModule = FModuleManager::Get().GetModulePtr<UE::Workspace::IWorkspaceEditorModule>("WorkspaceEditor"))
+		{
+			WorkspaceEditorModule->UnregisterObjectDocumentType(UUAFLayerStack::StaticClass()->GetClassPathName());
+			WorkspaceEditorModule->UnregisterWorkspaceItemDetails(UE::Workspace::FOutlinerItemDetailsId(FUAFLayerStackWorkspaceOutlinerData::StaticStruct()->GetFName()));
+		}
+
+		// Unregister layer stack class from UAFEditor
+		if (UE::UAF::Editor::IAnimNextEditorModule* UAFEditorModule = FModuleManager::Get().GetModulePtr<UE::UAF::Editor::IAnimNextEditorModule>("UAFEditor"))
+		{
+			UAFEditorModule->RemoveWorkspaceSupportedAssetClass(UUAFLayerStack::StaticClass()->GetClassPathName());
+		}
+	}
+}
+
+TSharedRef<SWidget> FUAFLayeringEditorModule::MakeDocumentWidget(const UE::Workspace::FWorkspaceEditorContext& InContext)
+{
+	if (UUAFLayerStack* LayerStack = InContext.Document.GetTypedObject<UUAFLayerStack>())
+	{
+		return SNew(UE::UAF::Layering::SLayerStack)
+			.LayerStack(LayerStack)
+			.WorkspaceEditor(InContext.WorkspaceEditor);
+	}
+
+	return SNullWidget::NullWidget;
+}
+
+void FUAFLayeringEditorModule::GetBreadcrumbTrail(const UE::Workspace::FWorkspaceEditorContext& InContext, TArray<TSharedPtr<UE::Workspace::FWorkspaceBreadcrumb>>& OutBreadcrumbs)
+{
+	UUAFLayerStack* LayerStack = InContext.Document.GetTypedObject<UUAFLayerStack>();
+	if (!LayerStack)
+	{
+		UE_LOG(LogAnimation, Error, TEXT("Tried to get breadcrumb trail for LayerStack asset but object was not a LayerStack"));
+		return;
+	}
+
+	// Create Breadcrumb
+	const TSharedPtr<UE::Workspace::FWorkspaceBreadcrumb>& BreadCrumb = OutBreadcrumbs.Add_GetRef(MakeShared<UE::Workspace::FWorkspaceBreadcrumb>());
+
+	// Set Breadcrumb Label
+	BreadCrumb->OnGetLabel = UE::Workspace::FWorkspaceBreadcrumb::FOnGetBreadcrumbLabel::CreateLambda([LayerStackName = LayerStack->GetFName()]
+		{
+			return FText::FromName(LayerStackName);
+		});
+
+	// Set Breadcrumb dirty flag 
+	TWeakObjectPtr<UUAFLayerStack> WeakLayerStack = LayerStack;
+	BreadCrumb->CanSave = UE::Workspace::FWorkspaceBreadcrumb::FCanSaveBreadcrumb::CreateLambda([WeakLayerStack]
+		{
+			return WeakLayerStack.IsValid() && WeakLayerStack->GetPackage()->IsDirty();
+		});
+
+	// Set breadcrumb click behavior
+	TWeakPtr<UE::Workspace::IWorkspaceEditor> WeakWorkspaceEditor = InContext.WorkspaceEditor;
+	BreadCrumb->OnClicked = UE::Workspace::FWorkspaceBreadcrumb::FOnBreadcrumbClicked::CreateLambda([WeakWorkspaceEditor, Export = InContext.Document.Export]
+		{
+			if (const TSharedPtr<UE::Workspace::IWorkspaceEditor> SharedWorkspaceEditor = WeakWorkspaceEditor.Pin())
+			{
+				SharedWorkspaceEditor->OpenExports({ Export });
+			}
+		});
+}
+
+IMPLEMENT_MODULE(FUAFLayeringEditorModule, UAFLayeringEditor)
+
+}
+
+#undef LOCTEXT_NAMESPACE
+	
diff --git a/Engine/Plugins/Experimental/UAF/UAFLayering/Source/UAFLayeringEditor/Private/UAFLayeringEditorModule.h b/Engine/Plugins/Experimental/UAF/UAFLayering/Source/UAFLayeringEditor/Private/UAFLayeringEditorModule.h
new file mode 100644
--- /dev/null
+++ b/Engine/Plugins/Experimental/UAF/UAFLayering/Source/UAFLayeringEditor/Private/UAFLayeringEditorModule.h
@@ -0,0 +1,25 @@
+// Copyright Epic Games, Inc. All Rights Reserved.
+
+#pragma once
+
+#include "Modules/ModuleManager.h"
+#include "IWorkspaceEditorModule.h"
+#include "AnimNextRigVMAssetEditorData.h"
+
+
+namespace UE::UAF::LayeringEditor
+{
+class FUAFLayeringEditorModule : public IModuleInterface
+{
+
+public:
+	/** IModuleInterface implementation */
+	virtual void StartupModule() override;
+	virtual void ShutdownModule() override;
+
+private:
+	TSharedRef<SWidget> MakeDocumentWidget(const UE::Workspace::FWorkspaceEditorContext& InContext);
+	void GetBreadcrumbTrail(const UE::Workspace::FWorkspaceEditorContext& InContext, TArray<TSharedPtr<UE::Workspace::FWorkspaceBreadcrumb>>& OutBreadcrumbs);
+};
+}
+
diff --git a/Engine/Plugins/Experimental/UAF/UAFLayering/Source/UAFLayeringEditor/Private/Workspace/LayerStackItemDetails.cpp b/Engine/Plugins/Experimental/UAF/UAFLayering/Source/UAFLayeringEditor/Private/Workspace/LayerStackItemDetails.cpp
new file mode 100644
--- /dev/null
+++ b/Engine/Plugins/Experimental/UAF/UAFLayering/Source/UAFLayeringEditor/Private/Workspace/LayerStackItemDetails.cpp
@@ -0,0 +1,15 @@
+// Copyright Epic Games, Inc. All Rights Reserved.
+
+#include "LayerStackItemDetails.h"
+
+#include "UAFLayeringStyle.h"
+#include "UAFStyle.h"
+
+namespace UE::UAF::LayeringEditor
+{
+const FSlateBrush* ULayerStackItemDetails::GetItemIcon(const FWorkspaceOutlinerItemExport& Export) const
+{
+	return FUAFLayeringStyle::Get().GetBrush(TEXT("ClassIcon.UAFLayerStack"));
+}
+	
+}
diff --git a/Engine/Plugins/Experimental/UAF/UAFLayering/Source/UAFLayeringEditor/Private/Workspace/LayerStackItemDetails.h b/Engine/Plugins/Experimental/UAF/UAFLayering/Source/UAFLayeringEditor/Private/Workspace/LayerStackItemDetails.h
new file mode 100644
--- /dev/null
+++ b/Engine/Plugins/Experimental/UAF/UAFLayering/Source/UAFLayeringEditor/Private/Workspace/LayerStackItemDetails.h
@@ -0,0 +1,15 @@
+// Copyright Epic Games, Inc. All Rights Reserved.
+
+#pragma once
+
+#include "IWorkspaceOutlinerItemDetails.h"
+
+namespace UE::UAF::LayeringEditor
+{
+class ULayerStackItemDetails : public UE::Workspace::IWorkspaceOutlinerItemDetails
+{
+public:
+	virtual const FSlateBrush* GetItemIcon(const FWorkspaceOutlinerItemExport& Export) const override;
+};
+}
+
diff --git a/Engine/Plugins/Experimental/UAF/UAFLayering/Source/UAFLayeringEditor/UAFLayeringEditor.Build.cs b/Engine/Plugins/Experimental/UAF/UAFLayering/Source/UAFLayeringEditor/UAFLayeringEditor.Build.cs
new file mode 100644
--- /dev/null
+++ b/Engine/Plugins/Experimental/UAF/UAFLayering/Source/UAFLayeringEditor/UAFLayeringEditor.Build.cs
@@ -0,0 +1,58 @@
+// Copyright Epic Games, Inc. All Rights Reserved.
+
+using UnrealBuildTool;
+using UnrealBuildTool.Rules;
+
+public class UAFLayeringEditor : ModuleRules
+{
+	public UAFLayeringEditor(ReadOnlyTargetRules Target) : base(Target)
+	{
+		PCHUsage = ModuleRules.PCHUsageMode.UseExplicitOrSharedPCHs;
+		
+		PublicDependencyModuleNames.AddRange(
+			new string[]
+			{
+				"Core",
+				"CoreUObject",
+				"Engine",
+				"Slate",
+				"SlateCore",
+				"WorkspaceEditor",
+				"RigVMDeveloper",
+				"UAFUncookedOnly",
+				"InputCore"
+	
+				// ... add other public dependencies that you statically link with here ...
+			}
+		);
+		
+		PrivateDependencyModuleNames.AddRange(
+			new string[]
+			{
+				"UAFAnimGraph",
+				"UnrealEd",
+				"AssetDefinition",
+				"UAFLayering",
+				"UAF",
+				"RigVM",
+				"EditorInteractiveToolsFramework",
+				"InteractiveToolsFramework",
+				"EditorFramework",
+				"UAFEditor",
+				"UAFLayeringUncookedOnly"
+
+
+				// ... add private dependencies that you statically link with here ...	
+			}
+			);
+
+		
+
+		PrivateIncludePathModuleNames.AddRange(
+			new string[]
+			{
+				"UAFEditor",
+			}
+		);
+	}
+}
diff --git a/Engine/Plugins/Experimental/UAF/UAFLayering/Source/UAFLayeringUncookedOnly/Internal/Layers/UAFBaseLayer.h b/Engine/Plugins/Experimental/UAF/UAFLayering/Source/UAFLayeringUncookedOnly/Internal/Layers/UAFBaseLayer.h
new file mode 100644
--- /dev/null
+++ b/Engine/Plugins/Experimental/UAF/UAFLayering/Source/UAFLayeringUncookedOnly/Internal/Layers/UAFBaseLayer.h
@@ -0,0 +1,20 @@
+// Copyright Epic Games, Inc. All Rights Reserved.
+
+#pragma once
+
+#include "CoreMinimal.h"
+#include "UAFLayer.h"
+#include "UAFBaseLayer.generated.h"
+
+// The editor representation of a base layer in a layer stack
+UCLASS(editinlinenew, meta =(Hidden))
+class UAFLAYERINGUNCOOKEDONLY_API UUAFBaseLayer : public UUAFLayer
+{
+	GENERATED_BODY()
+
+public:
+	UUAFBaseLayer();
+	virtual TSharedRef<SWidget> CreateLayerWidget() override;
+};
+
+
diff --git a/Engine/Plugins/Experimental/UAF/UAFLayering/Source/UAFLayeringUncookedOnly/Internal/Layers/UAFLayer.h b/Engine/Plugins/Experimental/UAF/UAFLayering/Source/UAFLayeringUncookedOnly/Internal/Layers/UAFLayer.h
new file mode 100644
--- /dev/null
+++ b/Engine/Plugins/Experimental/UAF/UAFLayering/Source/UAFLayeringUncookedOnly/Internal/Layers/UAFLayer.h
@@ -0,0 +1,82 @@
+// Copyright Epic Games, Inc. All Rights Reserved.
+
+#pragma once
+
+
+#include "CoreMinimal.h"
+#include "LayeringUncookedOnlyTypes.h"
+#include "Layers/UAFLayerBlendProviderBase.h"
+#include "Layers/UAFLayerContentProviderBase.h"
+#include "UAFLayerStack_EditorData.h"
+#include "UAFLayer.generated.h"
+
+class UAnimNextController;
+
+enum class EAnimNextEditorDataNotifType : uint8;
+
+
+// The editor representation of a layer in a layer stack
+UCLASS(editinlinenew)
+class UAFLAYERINGUNCOOKEDONLY_API UUAFLayer : public UObject
+{
+	GENERATED_BODY()
+
+public:
+	UUAFLayer();
+	
+	// Creates the widget to represent this layer within the layer stack 
+	virtual TSharedRef<SWidget> CreateLayerWidget();
+	
+	// Responsible for creating all traits that describe this layer 
+	// Per default this will request the LayerContentProvider and LayerBlendProvider to specify the exact implementation 
+	// @return Returns the output pin of the created layer. This can be nullptr if the creation failed. 
+	virtual URigVMPin* CreateLayerGraphContentTraits(UE::UAF::Layering::FLayerCreationContext& LayerCreationContext);
+	
+	// Toggles the preview state of this layer 
+	void TogglePreviewVisibility();
+	
+	// Gets the current layer name 
+	FName GetLayerName() const;
+	
+	// Renames this layer 
+	void RenameLayer(const FName NewName);
+	
+	// Checks if the passed in name is a valid name for this layer 
+	bool IsLayerNameValid(const FName NewName) const;
+	
+	EUAFLayerState GetLayerState() const;
+	void SetLayerState(const EUAFLayerState InLayerState);
+	
+	bool IsLayerContentProviderValid() const;
+	TInstancedStruct<FUAFLayerBlendProviderBase>& GetLayerBlendProvider();
+	void SetLayerBlendProvider(const TInstancedStruct<FUAFLayerBlendProviderBase>& InLayerBlendProvider);
+	
+	bool IsLayerBlendProviderValid() const;
+	TInstancedStruct<FUAFLayerContentProviderBase>& GetLayerContentProvider();
+	void SetLayerContentProvider(const TInstancedStruct<FUAFLayerContentProviderBase>& InLayerContentProvider);
+	
+	void BroadcastModified(EAnimNextEditorDataNotifType InType, bool bNotifyLayerLayoutChanges = false);
+
+#if WITH_EDITOR
+	virtual void PreEditChange(FProperty* PropertyAboutToChange) override;
+	virtual void PostEditChangeChainProperty(struct FPropertyChangedChainEvent& PropertyChangedEvent) override;
+#endif
+	
+protected:
+	UPROPERTY(EditAnywhere, Category = "Layer")
+	// The current state of the layer. 
+	EUAFLayerState LayerState = EUAFLayerState::Enabled;
+
+	// The name of this layer. It has to be unique within one layer stack asset. 
+	UPROPERTY(EditAnywhere, Category = "Layer")
+	FName LayerName;
+	
+	UPROPERTY(EditAnywhere, Category = "Layer", meta = (ExcludeBaseStruct))
+	TInstancedStruct<FUAFLayerContentProviderBase> LayerContentProvider;
+	
+	UPROPERTY(EditAnywhere, Category = "Layer", meta = (ExcludeBaseStruct))
+	TInstancedStruct<FUAFLayerBlendProviderBase> LayerBlendProvider;
+	
+	FName PreRenameLayerName = NAME_None;
+};
+
diff --git a/Engine/Plugins/Experimental/UAF/UAFLayering/Source/UAFLayeringUncookedOnly/Internal/Layers/UAFLayerAssetProvider.h b/Engine/Plugins/Experimental/UAF/UAFLayering/Source/UAFLayeringUncookedOnly/Internal/Layers/UAFLayerAssetProvider.h
new file mode 100644
--- /dev/null
+++ b/Engine/Plugins/Experimental/UAF/UAFLayering/Source/UAFLayeringUncookedOnly/Internal/Layers/UAFLayerAssetProvider.h
@@ -0,0 +1,37 @@
+// Copyright Epic Games, Inc. All Rights Reserved.
+
+#pragma once
+
+#include "CoreMinimal.h"
+#include "UAFLayerContentProviderBase.h"
+#include "Factory/AnimNextFactoryParams.h"
+#include "UAFLayerAssetProvider.generated.h"
+
+class UUAFLayer;
+class SWidget;
+
+// Layer content provider to support any pose outputting asset 
+USTRUCT()
+struct FUAFLayerAssetProvider : public FUAFLayerContentProviderBase
+{
+	GENERATED_BODY()
+
+public: 
+	virtual URigVMPin* CreateLayerContentTrait(UE::UAF::Layering::FLayerCreationContext& LayerCreationContext) override;
+	virtual TSharedRef<SWidget> CreateLayerContentWidget(UUAFLayer* InLayer) override;
+
+	void SetLayerAsset(UObject* InNewAsset);
+	
+private:
+	void OnLayerAssetSelectionChanged(const FAssetData& InAssetData);
+	
+private:
+	UPROPERTY(EditAnywhere, Category = "Layer|Content", meta = (GetAllowedClasses = "/Script/UAFAnimGraph.AnimNextAnimGraphSettings:GetAllowedAssetClasses"))
+	TObjectPtr<UObject> LayerAsset = nullptr;
+
+	UPROPERTY(EditAnywhere, Category = "Layer|Content", meta=(FactorySourceProperty = "LayerAsset"))
+	FAnimNextFactoryParams Parameters;
+	
+	UPROPERTY()
+	TWeakObjectPtr<UUAFLayer> OuterLayer = nullptr;
+};
diff --git a/Engine/Plugins/Experimental/UAF/UAFLayering/Source/UAFLayeringUncookedOnly/Internal/Layers/UAFLayerBlendProviderBase.h b/Engine/Plugins/Experimental/UAF/UAFLayering/Source/UAFLayeringUncookedOnly/Internal/Layers/UAFLayerBlendProviderBase.h
new file mode 100644
--- /dev/null
+++ b/Engine/Plugins/Experimental/UAF/UAFLayering/Source/UAFLayeringUncookedOnly/Internal/Layers/UAFLayerBlendProviderBase.h
@@ -0,0 +1,38 @@
+// Copyright Epic Games, Inc. All Rights Reserved.
+
+#pragma once
+
+
+#include "CoreMinimal.h"
+#include "LayeringUncookedOnlyTypes.h"
+#include "Widgets/SNullWidget.h"
+#include "UAFLayerBlendProviderBase.generated.h"
+
+class URigVMPin;
+class UUAFLayer;
+class SWidget;
+
+USTRUCT(meta=(Hidden))
+struct FUAFLayerBlendProviderBase
+{
+	GENERATED_BODY()
+
+public:
+	FUAFLayerBlendProviderBase() = default;
+	virtual ~FUAFLayerBlendProviderBase() = default;
+	
+	// Responsible for creating traits to represent this blend and link inputs to the relevant traits.
+	// @return The final output pin of this layer blend which can be linked as input by other layers or be linked to the result of the graph
+	virtual URigVMPin* CreateBlendGraphTrait(UE::UAF::Layering::FLayerCreationContext& LayerCreationContext) {return nullptr;};
+	
+	// Creates the widget to represent this blend in the layer stack editor 
+	virtual TSharedRef<SWidget> CreateLayerBlendWidget(UUAFLayer* OuterLayer) { return SNullWidget::NullWidget;} ;
+	
+	// Optional override of the containing layers background brush 
+	virtual const FSlateBrush* GetOverrideLayerBackground() const { return nullptr; };
+	
+	// Optional override of the containing layers indicator color
+	virtual bool GetOverrideIndicatorColor(FSlateColor& OutSlateColor) const {return false;};
+};
+
+
diff --git a/Engine/Plugins/Experimental/UAF/UAFLayering/Source/UAFLayeringUncookedOnly/Internal/Layers/UAFLayerContentProviderBase.h b/Engine/Plugins/Experimental/UAF/UAFLayering/Source/UAFLayeringUncookedOnly/Internal/Layers/UAFLayerContentProviderBase.h
new file mode 100644
--- /dev/null
+++ b/Engine/Plugins/Experimental/UAF/UAFLayering/Source/UAFLayeringUncookedOnly/Internal/Layers/UAFLayerContentProviderBase.h
@@ -0,0 +1,29 @@
+// Copyright Epic Games, Inc. All Rights Reserved.
+
+#pragma once
+
+#include "CoreMinimal.h"
+#include "Widgets/SNullWidget.h"
+#include "LayeringUncookedOnlyTypes.h"
+#include "UAFLayerContentProviderBase.generated.h"
+
+class URigVMPin;
+class SWidget;
+class UUAFLayer;
+
+USTRUCT(meta=(Hidden))
+struct FUAFLayerContentProviderBase 
+{
+	GENERATED_BODY()
+public:		
+	FUAFLayerContentProviderBase() = default;
+	virtual ~FUAFLayerContentProviderBase() = default;
+
+	// Responsible for creating the traits to play the content of this layer 
+	// @return The output pin of this layers content trait stack which can be linked to the blend trait stack of this layer 
+	virtual URigVMPin* CreateLayerContentTrait(UE::UAF::Layering::FLayerCreationContext& LayerCreationContext)  {return nullptr;} ;
+	
+	// Responsible for creating the widget to represent this layers content in the layer stack editor 
+	virtual TSharedRef<SWidget> CreateLayerContentWidget(UUAFLayer* InLayer) {return SNullWidget::NullWidget;} ;  
+};
+
diff --git a/Engine/Plugins/Experimental/UAF/UAFLayering/Source/UAFLayeringUncookedOnly/Internal/Layers/UAFLayerDefaultBlendProvider.h b/Engine/Plugins/Experimental/UAF/UAFLayering/Source/UAFLayeringUncookedOnly/Internal/Layers/UAFLayerDefaultBlendProvider.h
new file mode 100644
--- /dev/null
+++ b/Engine/Plugins/Experimental/UAF/UAFLayering/Source/UAFLayeringUncookedOnly/Internal/Layers/UAFLayerDefaultBlendProvider.h
@@ -0,0 +1,105 @@
+// Copyright Epic Games, Inc. All Rights Reserved.
+
+#pragma once
+
+
+#include "CoreMinimal.h"
+#include "UAFLayerBlendProviderBase.h"
+#include "IStructureDataProvider.h"
+#include "UAFLayeringTypes.h"
+#include "StructUtils/InstancedStruct.h"
+
+#include "UAFLayerDefaultBlendProvider.generated.h"
+
+class UAnimNextController;
+class UHierarchyTable;
+class URigVMPin;
+class UUAFLayer;
+class FUAFDefaultLayerBlendStructureDataProvider;
+
+USTRUCT()
+struct FUAFDefaultBlendProvider : public FUAFLayerBlendProviderBase
+{
+	GENERATED_BODY()
+
+public:
+	// Begin FUAFLayerBlendProvider
+	virtual URigVMPin* CreateBlendGraphTrait(UE::UAF::Layering::FLayerCreationContext& LayerCreationContext) override;
+	virtual TSharedRef<SWidget> CreateLayerBlendWidget(UUAFLayer* InLayer) override;
+	
+	// TODO: We can abstract this out to some layer override struct or to delegates we can pass into the widget 
+	virtual const FSlateBrush* GetOverrideLayerBackground() const override;
+	virtual bool GetOverrideIndicatorColor(FSlateColor& OutSlateColor) const override;
+	// End FUAFLayerBlendProvider
+
+public:
+	// TODO: This should be a binding
+	// If the layer is currently enabled at all 
+	UPROPERTY(EditAnywhere, Category = "Layer|BlendSettings")
+	bool bLayerEnabled = true;
+	
+	// The blend mode this layer uses to blend its result with the pose result of the previous layer 
+	UPROPERTY(EditAnywhere, Category = "Layer")
+	EUAFLayerBlendMode BlendMode = EUAFLayerBlendMode::Blend;
+	
+	// If not custom curve is used, which easing method should be used when this layer is blending in or out  
+	UPROPERTY(EditAnywhere, Category = "Layer|BlendSettings", meta = (EditCondition = "BlendCurve == nullptr", EditConditionHides))
+	EAlphaBlendOption BlendOption = EAlphaBlendOption::Linear;
+	
+	// TODO: This should be a binding
+	// The weight to use when blending this layer in. 0 == this layer won't get blended in at all, 1 == 100% of the layer gets blended in 
+	UPROPERTY(EditAnywhere, Category = "Layer", meta=(ClampMin = "0.0", ClampMax = "1.0", UIMin = "0.0", UIMax = "1.0"))
+	float LayerWeight = 1.0f;
+	
+	// TODO: This should be a binding
+	// How long it should take to blend this layer in if it has been disabled previously 
+	UPROPERTY(EditAnywhere, Category = "Layer|BlendSettings")
+	float LayerBlendInTime = 0.5f;
+
+	// TODO: This should be a binding
+	// How long it should take to blend this layer out 
+	UPROPERTY(EditAnywhere, Category = "Layer|BlendSettings")
+	float LayerBlendOutTime = 0.5f;
+	
+	// TODO: This should be a binding
+	// A optional blend mask to apply when blending this layers pose result with the previous layer result 
+	UPROPERTY(EditAnywhere, Category = "Layer")
+	TObjectPtr<UHierarchyTable> BlendMask = nullptr;
+	
+	// TODO: This should be a binding
+	// The custom blend curve to use when blending this layer in or out 
+	UPROPERTY(EditAnywhere, Category = "Layer|BlendSettings")
+	TObjectPtr<class UCurveFloat> BlendCurve = nullptr;
+	
+	// TODO: This should be a binding
+	// The blend profile to use to add more fine grain control over how certain bones should blend in or out 
+	// This can be used additionally to the blend mask that is used to apply this layer 
+	UPROPERTY(EditAnywhere, Category = "Layer|BlendSettings")
+	TObjectPtr<UHierarchyTable> TimeBasedBlendProfile = nullptr;
+	
+	// TODO: In the future we should be able to read this out of the blend mask 
+	// If a blend mask is applied, this is the named set for abstract hierarchies that mask relates to
+	UPROPERTY(EditAnywhere, Category = "Layer", meta=(EditCondition="BlendMask != nullptr"))
+	FName SetName = NAME_None;
+};
+
+namespace UE::UAF::Layering
+{
+class FDefaultLayerBlendStructureDataProvider : public IStructureDataProvider
+{
+public:
+	FDefaultLayerBlendStructureDataProvider(TInstancedStruct<FUAFLayerBlendProviderBase>* InBlendProvider)
+		: BlendProvider(InBlendProvider)
+	{}
+
+	virtual bool IsValid() const override;
+	virtual const UStruct* GetBaseStructure() const override;
+	virtual void GetInstances(TArray<TSharedPtr<FStructOnScope>>& OutInstances, const UStruct* ExpectedBaseStructure) const override;
+
+protected:
+	TInstancedStruct<FUAFLayerBlendProviderBase>* BlendProvider = nullptr;
+};
+}
+
+
+
diff --git a/Engine/Plugins/Experimental/UAF/UAFLayering/Source/UAFLayeringUncookedOnly/Internal/UAFLayerStack_EditorData.h b/Engine/Plugins/Experimental/UAF/UAFLayering/Source/UAFLayeringUncookedOnly/Internal/UAFLayerStack_EditorData.h
new file mode 100644
--- /dev/null
+++ b/Engine/Plugins/Experimental/UAF/UAFLayering/Source/UAFLayeringUncookedOnly/Internal/UAFLayerStack_EditorData.h
@@ -0,0 +1,67 @@
+// Copyright Epic Games, Inc. All Rights Reserved.
+
+#pragma once
+
+#include "LayeringUncookedOnlyTypes.h"
+#include "Graph/AnimNextAnimationGraph_EditorData.h"
+#include "UAFLayerStack_EditorData.generated.h"
+
+class UUAFLayer;
+class UUAFBaseLayer;
+class UUAFLayerStack;
+
+#define UE_API UAFLAYERINGUNCOOKEDONLY_API
+
+DECLARE_MULTICAST_DELEGATE(FLayerLayoutChangedDelegate);
+
+UCLASS()
+class UUAFLayerStack_EditorData : public UUAFAnimGraph_EditorData
+{
+	GENERATED_BODY()
+
+public:
+	int32 GetNumLayers() const;
+	TArray<TObjectPtr<UUAFLayer>> GetAllLayers() const;
+	int32 GetIndexForLayer(UUAFLayer* Layer) const;
+	bool IsBaseLayer(UUAFLayer* Layer) const;
+	bool IsLastLayer(UUAFLayer* Layer) const;
+	bool IsLayerNameValid(const FName NewLayerName) const;
+	
+	void MoveLayerUp(UUAFLayer* LayerToMove);
+	void MoveLayerDown(UUAFLayer* LayerToMove);
+	void DeleteLayer(UUAFLayer* LayerToDelete);
+	void SetLayerState(UUAFLayer* Layer, EUAFLayerState LayerState);
+	
+	UUAFLayer* CreateDefaultAssetBasedLayer(const FAssetData& InAsset);
+	void NotifyLayerLayoutChanged();
+	
+protected:
+	// Begin UAnimNextAnimationGraph_EditorData interface
+	UE_API virtual void Initialize(bool bRecompileVM) override;
+	UE_API virtual TSubclassOf<UAssetUserData> GetAssetUserDataClass() const override;
+	UE_API virtual TConstArrayView<TSubclassOf<UAnimNextRigVMAssetEntry>> GetEntryClasses() const override;
+	virtual void OnPreCompileGetProgrammaticGraphs(const FRigVMCompileSettings& InSettings, FAnimNextGetGraphCompileContext& OutCompileContext) override;
+	// End UAnimNextAnimationGraph_EditorData interface
+	
+	// Begin UObject interface
+	UE_API virtual void PostEditChangeProperty(FPropertyChangedEvent& PropertyChangedEvent) override;
+	// End UObject interface
+	
+public:
+	UPROPERTY()
+	TWeakObjectPtr<URigVMGraph> CreatedGraph = nullptr;
+	
+	FLayerLayoutChangedDelegate OnLayerLayoutChanged;
+
+private:
+	UPROPERTY(Instanced, VisibleAnywhere, Category = "LayerStack")
+	TObjectPtr<UUAFBaseLayer> BaseLayer;
+
+	UPROPERTY(Instanced, EditAnywhere, Category = "LayerStack")
+	TArray<TObjectPtr<UUAFLayer>> Layers;
+	
+	friend class UUAFLayerStackFactory;
+	friend class UUAFLayer;
+};
+
+#undef UE_API
diff --git a/Engine/Plugins/Experimental/UAF/UAFLayering/Source/UAFLayeringUncookedOnly/Internal/UAFLayeringStyle.h b/Engine/Plugins/Experimental/UAF/UAFLayering/Source/UAFLayeringUncookedOnly/Internal/UAFLayeringStyle.h
new file mode 100644
--- /dev/null
+++ b/Engine/Plugins/Experimental/UAF/UAFLayering/Source/UAFLayeringUncookedOnly/Internal/UAFLayeringStyle.h
@@ -0,0 +1,15 @@
+// Copyright Epic Games, Inc. All Rights Reserved.
+
+#pragma once
+
+#include "Styling/SlateStyle.h"
+
+class FUAFLayeringStyle final : public FSlateStyleSet
+{
+public:
+	static UAFLAYERINGUNCOOKEDONLY_API FUAFLayeringStyle& Get();
+
+private:
+	FUAFLayeringStyle();
+	~FUAFLayeringStyle();
+};
\ No newline at end of file
diff --git a/Engine/Plugins/Experimental/UAF/UAFLayering/Source/UAFLayeringUncookedOnly/Internal/Widgets/SDefaultBlendProvider.h b/Engine/Plugins/Experimental/UAF/UAFLayering/Source/UAFLayeringUncookedOnly/Internal/Widgets/SDefaultBlendProvider.h
new file mode 100644
--- /dev/null
+++ b/Engine/Plugins/Experimental/UAF/UAFLayering/Source/UAFLayeringUncookedOnly/Internal/Widgets/SDefaultBlendProvider.h
@@ -0,0 +1,36 @@
+// Copyright Epic Games, Inc. All Rights Reserved.
+
+#pragma once
+
+#include "CoreMinimal.h"
+#include "ISinglePropertyView.h"
+#include "Widgets/SCompoundWidget.h"
+
+namespace UE::UAF::Layering
+{
+	class FDefaultLayerBlendStructureDataProvider;
+}
+
+class ISinglePropertyView;
+class UUAFLayer;
+
+namespace UE::UAF::Layering
+{
+class SDefaultBlendProvider : public SCompoundWidget
+{
+public:
+	SLATE_BEGIN_ARGS(SDefaultBlendProvider) {}
+		SLATE_ARGUMENT(UUAFLayer*, Layer)
+	SLATE_END_ARGS()
+
+	/** Constructs this widget with InArgs */
+	void Construct(const FArguments& InArgs);
+	
+private:
+	TSharedPtr<ISinglePropertyView> GenerateSinglePropertyView(const FName PropertyName, const TSharedRef<FDefaultLayerBlendStructureDataProvider> StructDataProvider, const FSinglePropertyParams& Params) const;
+	
+private:
+	TWeakObjectPtr<UUAFLayer> WeakLayer = nullptr;
+};
+
+}
diff --git a/Engine/Plugins/Experimental/UAF/UAFLayering/Source/UAFLayeringUncookedOnly/Internal/Widgets/SLayer.h b/Engine/Plugins/Experimental/UAF/UAFLayering/Source/UAFLayeringUncookedOnly/Internal/Widgets/SLayer.h
new file mode 100644
--- /dev/null
+++ b/Engine/Plugins/Experimental/UAF/UAFLayering/Source/UAFLayeringUncookedOnly/Internal/Widgets/SLayer.h
@@ -0,0 +1,39 @@
+// Copyright Epic Games, Inc. All Rights Reserved.
+
+#pragma once
+
+#include "CoreMinimal.h"
+#include "Widgets/SCompoundWidget.h"
+
+class UUAFLayer;
+
+namespace UE::UAF::Layering
+{
+class SLayer : public SCompoundWidget
+{
+public:
+	SLATE_BEGIN_ARGS(SLayer) {}
+		SLATE_ARGUMENT(UUAFLayer*, Layer)
+	SLATE_END_ARGS()
+	
+	void Construct(const FArguments& InArgs);
+
+private: 
+	const FSlateBrush* GetVisibilityBrushForLayer() const;
+	const FSlateBrush* GetAdvancedSettingBrushForLayer() const;
+	const FSlateBrush* GetBackgroundBrush() const;
+	const FSlateBrush* GetBackgroundOverlayBrush() const;
+	
+	FSlateColor GetIndicatorColor() const;
+	FText GetLayerNameTooltip() const;
+	FText GetLayerName() const;
+	bool IsSelected() const;
+	bool VerifyLayerName(const FText& Text, FText& OutErrorMessage) const;
+	void OnLayerNameCommitted(const FText& NewText, ETextCommit::Type CommitInfo) const;
+
+private: 
+	TWeakObjectPtr<UUAFLayer> Layer = nullptr;
+	bool bAdvancedSettingsExpanded = false;
+};
+	
+}
diff --git a/Engine/Plugins/Experimental/UAF/UAFLayering/Source/UAFLayeringUncookedOnly/Internal/Widgets/SLayerStack.h b/Engine/Plugins/Experimental/UAF/UAFLayering/Source/UAFLayeringUncookedOnly/Internal/Widgets/SLayerStack.h
new file mode 100644
--- /dev/null
+++ b/Engine/Plugins/Experimental/UAF/UAFLayering/Source/UAFLayeringUncookedOnly/Internal/Widgets/SLayerStack.h
@@ -0,0 +1,48 @@
+// Copyright Epic Games, Inc. All Rights Reserved.
+
+#pragma once
+
+#include "Modules/ModuleManager.h"
+#include "IWorkspaceEditorModule.h"
+#include "UAFLayerStack_EditorData.h"
+#include "Widgets/Views/SListView.h"
+
+class UUAFLayer;
+class FUICommandList;
+class UUAFLayerStack;
+
+namespace UE::Workspace
+{
+	class IWorkspaceEditor;
+}
+
+namespace UE::UAF::Layering
+{
+class SLayerStack : public SCompoundWidget
+{
+public:
+	SLATE_BEGIN_ARGS(SLayerStack) {}
+		SLATE_ARGUMENT(UUAFLayerStack*, LayerStack)
+		SLATE_ARGUMENT(TSharedPtr<UE::Workspace::IWorkspaceEditor>, WorkspaceEditor)
+		SLATE_EVENT(FSimpleDelegate, OnSelectedLayerChanged)
+	SLATE_END_ARGS()
+	
+	UAFLAYERINGUNCOOKEDONLY_API void Construct(const FArguments& InArgs);
+	UAFLAYERINGUNCOOKEDONLY_API virtual ~SLayerStack() override;
+	
+private:
+	TSharedRef<ITableRow> GenerateLayerRow(TWeakObjectPtr<UUAFLayer> Item, const TSharedRef<STableViewBase>& OwnerTable) const;
+	void OnListSelectionChanged(TWeakObjectPtr<UUAFLayer> SelectedItem, ESelectInfo::Type SelectInfo) const;
+	TSharedPtr<SWidget> OnConstructContextMenu() const;
+	void OnLayerLayoutChanged();
+	void CacheLayersForList();
+private:
+	TSharedPtr<SListView<TWeakObjectPtr<UUAFLayer>>> LayerListView = nullptr;
+	TArray<TWeakObjectPtr<UUAFLayer>> CachedLayerPointers;
+	
+	TSharedPtr<FUICommandList> CommandList = nullptr;
+	TWeakObjectPtr<UUAFLayerStack> LayerStack = nullptr;
+	TWeakPtr<UE::Workspace::IWorkspaceEditor> WeakWorkspaceEditor = nullptr;
+};
+
+}
diff --git a/Engine/Plugins/Experimental/UAF/UAFLayering/Source/UAFLayeringUncookedOnly/Internal/Workspace/UAFLayerStackWorkspaceAssetUserData.h b/Engine/Plugins/Experimental/UAF/UAFLayering/Source/UAFLayeringUncookedOnly/Internal/Workspace/UAFLayerStackWorkspaceAssetUserData.h
new file mode 100644
--- /dev/null
+++ b/Engine/Plugins/Experimental/UAF/UAFLayering/Source/UAFLayeringUncookedOnly/Internal/Workspace/UAFLayerStackWorkspaceAssetUserData.h
@@ -0,0 +1,20 @@
+// Copyright Epic Games, Inc. All Rights Reserved.
+
+#pragma once
+
+#include "AnimNextAnimGraphWorkspaceAssetUserData.h"
+#include "UObject/AssetRegistryTagsContext.h"
+#include "UAFLayerStackWorkspaceAssetUserData.generated.h"
+
+UCLASS()
+class UUAFLayerStackWorkspaceAssetUserData : public UAnimNextAnimGraphWorkspaceAssetUserData
+{
+	GENERATED_BODY()
+
+protected:
+	// Begin UAnimNextAnimGraphWorkspaceAssetUserData Interface
+	virtual void GetRootAssetExport(FAssetRegistryTagsContext Context) const override;
+	virtual void GetWorkspaceAssetExports(FAssetRegistryTagsContext Context) const override;
+	// End UAnimNextAnimGraphWorkspaceAssetUserData Interface
+
+};
\ No newline at end of file
diff --git a/Engine/Plugins/Experimental/UAF/UAFLayering/Source/UAFLayeringUncookedOnly/Internal/Workspace/UAFLayerStackWorkspaceOutlinerData.h b/Engine/Plugins/Experimental/UAF/UAFLayering/Source/UAFLayeringUncookedOnly/Internal/Workspace/UAFLayerStackWorkspaceOutlinerData.h
new file mode 100644
--- /dev/null
+++ b/Engine/Plugins/Experimental/UAF/UAFLayering/Source/UAFLayeringUncookedOnly/Internal/Workspace/UAFLayerStackWorkspaceOutlinerData.h
@@ -0,0 +1,15 @@
+// Copyright Epic Games, Inc. All Rights Reserved.
+
+#pragma once
+
+
+#include "AnimNextAssetWorkspaceAssetUserData.h"
+#include "UAFLayerStackWorkspaceOutlinerData.generated.h"
+
+USTRUCT()
+struct FUAFLayerStackWorkspaceOutlinerData : public FAnimNextRigVMAssetOutlinerData
+{
+	GENERATED_BODY()
+	
+	FUAFLayerStackWorkspaceOutlinerData() = default;
+};
\ No newline at end of file
diff --git a/Engine/Plugins/Experimental/UAF/UAFLayering/Source/UAFLayeringUncookedOnly/Private/Layers/UAFBaseLayer.cpp b/Engine/Plugins/Experimental/UAF/UAFLayering/Source/UAFLayeringUncookedOnly/Private/Layers/UAFBaseLayer.cpp
new file mode 100644
--- /dev/null
+++ b/Engine/Plugins/Experimental/UAF/UAFLayering/Source/UAFLayeringUncookedOnly/Private/Layers/UAFBaseLayer.cpp
@@ -0,0 +1,94 @@
+// Copyright Epic Games, Inc. All Rights Reserved.
+
+#pragma once
+
+#include "Layers/UAFBaseLayer.h"
+
+#include "UAFLayeringStyle.h"
+#include "Components/HorizontalBox.h"
+#include "Layers/UAFLayerAssetProvider.h"
+#include "Widgets/SOverlay.h"
+#include "Widgets/Images/SImage.h"
+#include "Widgets/Input/SEditableTextBox.h"
+#include "Widgets/Layout/SBorder.h"
+#include "Widgets/Text/STextBlock.h"
+
+#define LOCTEXT_NAMESPACE "FUAFLayeringEditorModule"
+
+UUAFBaseLayer::UUAFBaseLayer()
+{
+	SetLayerContentProvider(TInstancedStruct<FUAFLayerAssetProvider>::Make());
+	LayerBlendProvider.Reset();
+}
+
+TSharedRef<SWidget> UUAFBaseLayer::CreateLayerWidget()
+{
+	return SNew(SOverlay)
+		+ SOverlay::Slot()
+		.HAlign(HAlign_Fill)
+		.VAlign(VAlign_Fill)
+		[
+			SNew(SBorder)
+				.BorderImage(FUAFLayeringStyle::Get().GetBrush(TEXT("Layer.Background")))
+				[
+					SNew(SHorizontalBox)
+					+ SHorizontalBox::Slot()
+						.AutoWidth()
+						.Padding(5.0f, 10.f)
+						[
+							SNew(SVerticalBox)
+								+ SVerticalBox::Slot()
+								.AutoHeight()
+								.Padding(10.0f, 0.0f)
+								[
+									SNew(STextBlock)
+									.Text_Lambda([this]() -> FText
+									{
+										return FText::FromName(LayerName);
+									})
+								]
+								+ SVerticalBox::Slot()
+								.AutoHeight()
+								.Padding(10.0f, 5.0f)
+								[
+									LayerContentProvider.IsValid() ? LayerContentProvider.GetMutable().CreateLayerContentWidget(this) : SNullWidget::NullWidget
+								]
+						]
+					+ SHorizontalBox::Slot()
+						.AutoWidth()
+						[
+							SNew(SVerticalBox)
+							+SVerticalBox::Slot()
+							.Padding(20.0f, 30.0f, 10.0f, 5.0f)
+							.AutoHeight()
+								[
+									// Base Layer Cached Pose
+									SNew(STextBlock)
+										.Text(FText::FromString(TEXT("Base Cached Pose")))
+								]
+							+SVerticalBox::Slot()
+							.Padding(20.0f, 2.0f)
+							.AutoHeight()
+								[
+									SNew(SHorizontalBox)
+									+SHorizontalBox::Slot()
+									.Padding(0.0f, 2.0f)
+									.AutoWidth()
+									[
+										SNew(SImage)
+										.Image(FUAFLayeringStyle::Get().GetBrush(TEXT("Layer.PoseIconSmall")))
+									]
+									+SHorizontalBox::Slot()
+									.Padding(5.0f, 0.0f)
+									[
+										SNew(SEditableTextBox)
+												.Text(FText::FromString(TEXT("Base Pose Binding")))
+												.IsEnabled(false)
+									]
+								]
+						]
+				]
+		];
+}
+
+#undef LOCTEXT_NAMESPACE
diff --git a/Engine/Plugins/Experimental/UAF/UAFLayering/Source/UAFLayeringUncookedOnly/Private/Layers/UAFLayer.cpp b/Engine/Plugins/Experimental/UAF/UAFLayering/Source/UAFLayeringUncookedOnly/Private/Layers/UAFLayer.cpp
new file mode 100644
--- /dev/null
+++ b/Engine/Plugins/Experimental/UAF/UAFLayering/Source/UAFLayeringUncookedOnly/Private/Layers/UAFLayer.cpp
@@ -0,0 +1,190 @@
+// Copyright Epic Games, Inc. All Rights Reserved.
+
+#pragma once
+
+#include "Layers/UAFLayer.h"
+
+#include "UAFLayerStack_EditorData.h"
+#include "Layers/UAFLayerDefaultBlendProvider.h"
+#include "Layers/UAFLayerAssetProvider.h"
+#include "Widgets/SLayer.h"
+#include "Widgets/Views/STableRow.h"
+
+#define LOCTEXT_NAMESPACE "FUAFLayer"
+
+UUAFLayer::UUAFLayer()
+{
+	SetLayerContentProvider(TInstancedStruct<FUAFLayerAssetProvider>::Make());
+	SetLayerBlendProvider(TInstancedStruct<FUAFDefaultBlendProvider>::Make());
+}
+
+URigVMPin* UUAFLayer::CreateLayerGraphContentTraits(UE::UAF::Layering::FLayerCreationContext& LayerCreationContext)
+{
+	LayerCreationContext.Layer = this;
+	
+	// Create layer content traits 
+	URigVMPin* ContentOutputPin = LayerContentProvider.IsValid() ? LayerContentProvider.GetMutable().CreateLayerContentTrait(LayerCreationContext) : nullptr;
+	if (ContentOutputPin == nullptr)
+	{
+		LayerCreationContext.CompileSettings.ReportError(TEXT("Layer content traits could not be created successfully!"));
+		return nullptr;
+	}
+	
+	// The blend trait can be optional, some layers e.g. the base layer might only provide content but no blend 
+	// in these cases we pass through the layer content output 
+	if (LayerBlendProvider.IsValid())
+	{
+		// The layer content is the second input for the layering trait (ChildB) 
+		// The previous layer output will be at index 0, if applicable 
+		LayerCreationContext.LayerInputs[1] = ContentOutputPin;
+		
+		if (URigVMPin* BlendOutputPin = LayerBlendProvider.GetMutable().CreateBlendGraphTrait(LayerCreationContext))
+		{
+			return BlendOutputPin;
+		}
+	}
+	
+	return ContentOutputPin;
+}
+
+void UUAFLayer::TogglePreviewVisibility()
+{
+	switch (LayerState)
+	{
+	case EUAFLayerState::Enabled:
+		LayerState = EUAFLayerState::PreviewDisabled;
+		break;
+	case EUAFLayerState::PreviewDisabled:
+		LayerState = EUAFLayerState::Enabled;
+		break;
+	case EUAFLayerState::Disabled:
+	default:
+		// Do nothing 
+		break;
+	}
+}
+
+void UUAFLayer::BroadcastModified(EAnimNextEditorDataNotifType InType, bool bNotifyLayerLayoutChanges)
+{
+	if (UUAFLayerStack_EditorData* OuterEditorData = Cast<UUAFLayerStack_EditorData>(GetOuter()))
+	{
+		OuterEditorData->BroadcastModified(InType, this);
+		if (bNotifyLayerLayoutChanges)
+		{
+			OuterEditorData->NotifyLayerLayoutChanged();
+		}
+	}
+}
+
+#if WITH_EDITOR
+void UUAFLayer::PreEditChange(FProperty* PropertyAboutToChange)
+{
+	Super::PreEditChange(PropertyAboutToChange);
+	
+	if (PropertyAboutToChange && PropertyAboutToChange->GetName() == GET_MEMBER_NAME_CHECKED(UUAFLayer, LayerName))
+	{
+		PreRenameLayerName = LayerName;
+	}
+}
+
+void UUAFLayer::PostEditChangeChainProperty(FPropertyChangedChainEvent& PropertyChangedEvent)
+{
+	Super::PostEditChangeChainProperty(PropertyChangedEvent);
+	
+	bool bUpdateLayerLayout = false;
+	if (PropertyChangedEvent.MemberProperty)
+	{
+		if (PropertyChangedEvent.MemberProperty->GetName() == GET_MEMBER_NAME_CHECKED(UUAFLayer, LayerName))
+		{
+			if (!IsLayerNameValid(LayerName) && !PreRenameLayerName.IsNone())
+			{
+				UE_LOGFMT(LogAnimation, Error, "Layer Name {LayerName} is invalid, will reset to previous name {PreLayerName}", LayerName, PreRenameLayerName);
+				LayerName = PreRenameLayerName;
+			}
+		}
+
+		if (PropertyChangedEvent.MemberProperty->GetName() == GET_MEMBER_NAME_CHECKED(UUAFLayer, LayerBlendProvider) 
+			|| PropertyChangedEvent.MemberProperty->GetName() == GET_MEMBER_NAME_CHECKED(UUAFLayer, LayerContentProvider))
+		{
+			bUpdateLayerLayout = true;
+		}
+	}
+	
+	BroadcastModified(EAnimNextEditorDataNotifType::PropertyChanged, bUpdateLayerLayout);
+}
+#endif
+
+void UUAFLayer::RenameLayer(const FName NewName)
+{
+	LayerName = NewName;
+}
+
+bool UUAFLayer::IsLayerNameValid(const FName NewName) const
+{
+	if (NewName == LayerName)
+	{
+		return true;
+	}
+	
+	if (const UUAFLayerStack_EditorData* EditorData = Cast<UUAFLayerStack_EditorData>(GetOuter()))
+	{
+		return EditorData->IsLayerNameValid(NewName);
+	}
+	
+	return false;
+}
+
+EUAFLayerState UUAFLayer::GetLayerState() const
+{
+	return LayerState;
+}
+
+void UUAFLayer::SetLayerState(const EUAFLayerState InLayerState)
+{
+	this->LayerState = InLayerState;
+}
+
+bool UUAFLayer::IsLayerContentProviderValid() const
+{
+	return LayerContentProvider.IsValid();
+}
+
+FName UUAFLayer::GetLayerName() const
+{
+	return LayerName;
+}
+
+TInstancedStruct<FUAFLayerBlendProviderBase>& UUAFLayer::GetLayerBlendProvider()
+{
+	return LayerBlendProvider;
+}
+
+void UUAFLayer::SetLayerBlendProvider(const TInstancedStruct<FUAFLayerBlendProviderBase>& InLayerBlendProvider)
+{
+	LayerBlendProvider = InLayerBlendProvider;
+}
+
+bool UUAFLayer::IsLayerBlendProviderValid() const
+{
+	return LayerBlendProvider.IsValid();
+}
+
+TInstancedStruct<FUAFLayerContentProviderBase>& UUAFLayer::GetLayerContentProvider()
+{
+	return LayerContentProvider;
+}
+
+void UUAFLayer::SetLayerContentProvider(const TInstancedStruct<FUAFLayerContentProviderBase>& InLayerContentProvider)
+{
+	LayerContentProvider = InLayerContentProvider;
+}
+
+
+TSharedRef<SWidget> UUAFLayer::CreateLayerWidget()
+{
+	return SNew(UE::UAF::Layering::SLayer)
+		.Layer(this);
+	
+}
+
+#undef LOCTEXT_NAMESPACE
diff --git a/Engine/Plugins/Experimental/UAF/UAFLayering/Source/UAFLayeringUncookedOnly/Private/Layers/UAFLayerAssetProvider.cpp b/Engine/Plugins/Experimental/UAF/UAFLayering/Source/UAFLayeringUncookedOnly/Private/Layers/UAFLayerAssetProvider.cpp
new file mode 100644
--- /dev/null
+++ b/Engine/Plugins/Experimental/UAF/UAFLayering/Source/UAFLayeringUncookedOnly/Private/Layers/UAFLayerAssetProvider.cpp
@@ -0,0 +1,133 @@
+// Copyright Epic Games, Inc. All Rights Reserved.
+
+#pragma once
+
+#include "Layers/UAFLayerAssetProvider.h"
+
+#include "AnimNextAnimGraphSettings.h"
+#include "AnimNextController.h"
+#include "AnimNextRigVMAssetEditorData.h"
+#include "AssetThumbnail.h"
+#include "PropertyCustomizationHelpers.h"
+#include "Factory/AnimGraphFactory.h"
+#include "Graph/RigUnit_AnimNextTraitStack.h"
+#include "Layers/UAFLayer.h"
+#include "RigVMCore/RigVMStruct.h"
+#include "RigVMModel/Nodes/RigVMUnitNode.h"
+#include "TraitCore/Trait.h"
+#include "Traits/BlendStackTrait.h"
+#include "Traits/LayerAssetDataTraitData.h"
+#include "Traits/SubGraphHostTraitData.h"
+#include "Widgets/Layout/SBox.h"
+
+#define LOCTEXT_NAMESPACE "UAFLayerAssetContentProvider"
+
+URigVMPin* FUAFLayerAssetProvider::CreateLayerContentTrait(UE::UAF::Layering::FLayerCreationContext& LayerCreationContext)
+{
+	LayerCreationContext.LastNodeLocation += FVector2D(0.0f, 300.0f);
+	URigVMUnitNode* ContentHostTraitStackNode = LayerCreationContext.GraphController->AddUnitNode(FRigUnit_AnimNextTraitStack::StaticStruct(), FRigVMStruct::ExecuteName, LayerCreationContext.LastNodeLocation, FString(), false);
+	if (ContentHostTraitStackNode == nullptr)
+	{
+		LayerCreationContext.CompileSettings.ReportError(TEXT("Failed to create trait stack to host layer content."));
+		return nullptr;
+	}
+	
+	// TODO: We likely could unify this approach and not differentiate between graphs and other assets
+	if (UUAFAnimGraph* AnimationGraph = Cast<UUAFAnimGraph>(LayerAsset))
+	{
+		FAnimNextSubGraphHostTraitSharedData DefaultSubGraphData;
+		DefaultSubGraphData.AnimationGraph = AnimationGraph;
+		
+		const FName SubGraphHostTraitName = LayerCreationContext.GraphController->AddTraitStruct(ContentHostTraitStackNode, TInstancedStruct<FAnimNextSubGraphHostTraitSharedData>::Make(DefaultSubGraphData), INDEX_NONE, false, false);
+		if (SubGraphHostTraitName.IsNone())
+		{
+			LayerCreationContext.CompileSettings.ReportError(TEXT("Failed to create SubGraphHost Trait."));
+			return nullptr;
+		}
+		
+		// Expand trait pin per default 
+		URigVMPin* TraitPin = ContentHostTraitStackNode->FindTrait(SubGraphHostTraitName);
+		LayerCreationContext.GraphController->SetPinExpansion(TraitPin->GetPinPath(), true, false);
+	}
+	else
+	{
+		// Create Blend stack trait
+		const FName BlendStackTraitName = LayerCreationContext.GraphController->AddTraitStruct(ContentHostTraitStackNode, TInstancedStruct<FAnimNextBlendStackCoreTraitSharedData>::Make(), INDEX_NONE, false, false);
+		if (BlendStackTraitName.IsNone())
+		{
+			LayerCreationContext.CompileSettings.ReportError(TEXT("Failed to create Blend Stack Trait."));
+			return nullptr;
+		}
+		
+		// Create asset data provider trait 
+		FUAFLayerAssetDataTraitSharedData AssetDataDefaults;
+		AssetDataDefaults.Asset = LayerAsset;
+		AssetDataDefaults.Parameters = Parameters;
+		
+		const FName AssetDataTraitName = LayerCreationContext.GraphController->AddTraitStruct(ContentHostTraitStackNode, TInstancedStruct<FUAFLayerAssetDataTraitSharedData>::Make(AssetDataDefaults), INDEX_NONE, false, false);
+		if (AssetDataTraitName.IsNone())
+		{
+			LayerCreationContext.CompileSettings.ReportError(TEXT("Failed to create Layer Asset DataTrait."));
+			return nullptr;
+		}
+	}
+
+	URigVMPin* BlendTraitStackResult = ContentHostTraitStackNode->FindPin(GET_MEMBER_NAME_STRING_CHECKED(FRigUnit_AnimNextTraitStack, Result));
+	if (BlendTraitStackResult == nullptr)
+	{
+		LayerCreationContext.CompileSettings.ReportError(TEXT("Failed to retrieve Result pin from trait stack"));
+		return nullptr;
+	}
+	
+	return BlendTraitStackResult;
+}
+
+void FUAFLayerAssetProvider::SetLayerAsset(UObject* InNewAsset)
+{
+	LayerAsset = InNewAsset;
+}
+
+void FUAFLayerAssetProvider::OnLayerAssetSelectionChanged(const FAssetData& InAssetData)
+{
+	UObject* Obj = Cast<UObject>(InAssetData.GetAsset());
+	SetLayerAsset(Obj);
+	
+	if (OuterLayer.IsValid())
+	{
+		OuterLayer->BroadcastModified(EAnimNextEditorDataNotifType::PropertyChanged);
+	}
+}
+
+TSharedRef<SWidget> FUAFLayerAssetProvider::CreateLayerContentWidget(UUAFLayer* InLayer)
+{
+	OuterLayer = InLayer;
+
+	return SNew(SHorizontalBox)
+		+ SHorizontalBox::Slot()
+		.VAlign(EVerticalAlignment::VAlign_Center)
+		[
+			SNew(SObjectPropertyEntryBox)
+				.AllowedClass(UObject::StaticClass())
+				.DisplayThumbnail(true)
+				.DisplayCompactSize(false)
+				.OnShouldFilterAsset_Lambda([](const FAssetData& AssetData)
+					{
+						TArray<UClass*> AllowedClasses = UAnimNextAnimGraphSettings::GetAllowedAssetClasses();
+						return (AllowedClasses.Contains(AssetData.GetClass()) == false);
+					})
+				.ObjectPath_Lambda([this]()
+					{
+						if (LayerAsset)
+						{
+							return LayerAsset->GetPathName();
+						}
+						return FString();
+					})
+				.OnObjectChanged_Raw(this, &FUAFLayerAssetProvider::OnLayerAssetSelectionChanged)
+				.ThumbnailPool(MakeShared<FAssetThumbnailPool>(1))
+				
+		];
+}
+
+#undef LOCTEXT_NAMESPACE //UAFLayerAssetContentProvider
+
diff --git a/Engine/Plugins/Experimental/UAF/UAFLayering/Source/UAFLayeringUncookedOnly/Private/Layers/UAFLayerDefaultBlendProvider.cpp b/Engine/Plugins/Experimental/UAF/UAFLayering/Source/UAFLayeringUncookedOnly/Private/Layers/UAFLayerDefaultBlendProvider.cpp
new file mode 100644
--- /dev/null
+++ b/Engine/Plugins/Experimental/UAF/UAFLayering/Source/UAFLayeringUncookedOnly/Private/Layers/UAFLayerDefaultBlendProvider.cpp
@@ -0,0 +1,313 @@
+// Copyright Epic Games, Inc. All Rights Reserved.
+
+#pragma once
+
+#include "Layers/UAFLayerDefaultBlendProvider.h"
+
+#include "AnimNextController.h"
+#include "AnimNextRigVMAssetEditorData.h"
+#include "HierarchyTable.h"
+#include "ISinglePropertyView.h"
+#include "UAFLayeringStyle.h"
+#include "UAFLayerStack.h"
+#include "Graph/RigUnit_AnimNextTraitStack.h"
+#include "Layers/UAFLayer.h"
+#include "RigVMCore/RigVMStruct.h"
+#include "RigVMModel/Nodes/RigVMUnitNode.h"
+#include "TraitCore/Trait.h"
+#include "Traits/ApplyAdditiveTraitData.h"
+#include "Traits/ApplyNamedSetTraitData.h"
+#include "Traits/BlendLayerTraitData.h"
+#include "Traits/BlendTwoWayTraitData.h"
+#include "Traits/CachePoseTraitData.h"
+#include "Traits/LayerDataProviderTraitData.h"
+#include "Traits/PassthroughTraitData.h"
+#include "Widgets/SDefaultBlendProvider.h"
+
+#define LOCTEXT_NAMESPACE "UAFDefaultBlendProvider"
+
+
+URigVMPin* FUAFDefaultBlendProvider::CreateBlendGraphTrait(UE::UAF::Layering::FLayerCreationContext& LayerCreationContext)
+{
+	// In some cases we may have to create 2 trait stacks to allow for some of the more advanced features - mostly because currently we cannot inject tasks between child trait evals/updates 
+	// Case 1: Adding a named set to only evaluate what a layer needs, this needs to happen in between child trait evals (applies not to base but needs to apply to layer input) 
+	// Case 2: Cache Only Layers - We need to be able to pop the layer content keyframe and cache it at the right point to ensure the keyframe stack is in an expected state 
+	
+	// Create pre blend trait stack if needed (e.g. Cache Only or to provide name set optimization) 
+	if (BlendMask || BlendMode == EUAFLayerBlendMode::CacheOnly)
+	{
+		// Create Pre-Blend trait stack 
+		LayerCreationContext.LastNodeLocation += FVector2D(500.0f, 150.0f);
+		URigVMUnitNode* PreBlendTraitStack = LayerCreationContext.GraphController->AddUnitNode(FRigUnit_AnimNextTraitStack::StaticStruct(), FRigVMStruct::ExecuteName, LayerCreationContext.LastNodeLocation, FString(), false);
+		if (PreBlendTraitStack == nullptr)
+		{
+			LayerCreationContext.CompileSettings.ReportError(TEXT("Failed to create pre-blend trait stack."));
+			return nullptr;
+		}
+		
+		// The input pin for the pre-blend trait stack. This will depend on which traits are created
+		URigVMPin* PreBlendInput = nullptr;
+
+		if (BlendMask)
+		{
+			// Create Apply Named Set Trait - this is a base trait, so no need to create an additional pass through trait 
+			UE::UAF::FApplyNamedSetSharedData DefaultData;
+			DefaultData.SetName = SetName;
+			const FName ApplyNamedSetTraitName = LayerCreationContext.GraphController->AddTraitStruct(PreBlendTraitStack, TInstancedStruct<UE::UAF::FApplyNamedSetSharedData>::Make(DefaultData), INDEX_NONE, false, false);
+			if (ApplyNamedSetTraitName == NAME_None)
+			{
+				LayerCreationContext.CompileSettings.ReportError(TEXT("Failed to create Apply Named Set trait."));
+				return nullptr;
+			}
+			
+			// Hook up the input pin to the input of the Apply Named Set trait 
+			PreBlendInput = PreBlendTraitStack->FindTrait(ApplyNamedSetTraitName, GET_MEMBER_NAME_STRING_CHECKED(UE::UAF::FApplyNamedSetSharedData, Input));
+		}
+		
+		if (BlendMode == EUAFLayerBlendMode::CacheOnly)
+		{
+			// If we haven't created a trait yet to take care of the input we'll add a passthrough onto the trait stack 
+			if (PreBlendInput == nullptr)
+			{
+				// Create passthrough trait 
+				const FName PreBlendPassThroughTraitName = LayerCreationContext.GraphController->AddTraitStruct(PreBlendTraitStack, TInstancedStruct<FAnimNextPassthroughSharedData>::Make(), INDEX_NONE, false, false);
+		
+				// Populate the new input pin
+				PreBlendInput = PreBlendTraitStack->FindTrait(PreBlendPassThroughTraitName, GET_MEMBER_NAME_STRING_CHECKED(FAnimNextPassthroughSharedData, Input));
+			}
+			
+			// Create cache pose additive trait 
+			const FName CachePoseTraitName = LayerCreationContext.GraphController->AddTraitStruct(PreBlendTraitStack, TInstancedStruct<FUAFCachePoseTraitSharedData>::Make(), INDEX_NONE, false, false);
+			if (CachePoseTraitName == NAME_None)
+			{
+				LayerCreationContext.CompileSettings.ReportError(TEXT("Failed to create Cache Pose trait."));
+				return nullptr;
+			}
+		}
+
+		if (PreBlendInput && PreBlendTraitStack)
+		{
+			// Link current layer asset with passthrough
+			const bool bLayerContentLinkCreated = LayerCreationContext.GraphController->AddLink(LayerCreationContext.LayerInputs[1], PreBlendInput, false);
+			if (!bLayerContentLinkCreated)
+			{
+				LayerCreationContext.CompileSettings.ReportError(TEXT("Failed to link between layer content output and pre-blend trait"));
+				return nullptr;
+			}
+			
+			// Update input so the layer trait stack can link correctly 
+			URigVMPin* PreBlendTraitStackResult = PreBlendTraitStack->FindPin(GET_MEMBER_NAME_STRING_CHECKED(FRigUnit_AnimNextTraitStack, Result));
+			LayerCreationContext.LayerInputs[1] = PreBlendTraitStackResult;
+		}
+	}
+	
+	URigVMPin* ChildAInputPin = nullptr;
+	URigVMPin* ChildBInputPin = nullptr;
+	
+	// Create Blend Trait Stack Node
+	LayerCreationContext.LastNodeLocation += FVector2D(500.0f, -150.0f);
+	URigVMUnitNode* BlendTraitStackNode = LayerCreationContext.GraphController->AddUnitNode(FRigUnit_AnimNextTraitStack::StaticStruct(), FRigVMStruct::ExecuteName, LayerCreationContext.LastNodeLocation, FString(), false);
+	if (BlendTraitStackNode == nullptr)
+	{
+		LayerCreationContext.CompileSettings.ReportError(TEXT("Failed to create trait stack for the layer default blend implementation."));
+		return nullptr;
+	}
+	
+	FName BlendBaseTraitName = NAME_None;
+	const bool bUsesPerBoneData = BlendMask || TimeBasedBlendProfile;
+	if (BlendMode == EUAFLayerBlendMode::Additive)
+	{
+		if (bUsesPerBoneData)
+		{
+			// Not supported yet!
+			LayerCreationContext.CompileSettings.ReportWarning(TEXT("Per Bone application of additives in not supported yet! Will create a fully body additive trait instead"));
+		}
+		
+		FAnimNextApplyAdditiveTraitSharedData DefaultAdditiveTraitData;
+		DefaultAdditiveTraitData.Alpha = LayerWeight;
+		
+		// Create trait 
+		BlendBaseTraitName = LayerCreationContext.GraphController->AddTraitStruct(BlendTraitStackNode, TInstancedStruct<FAnimNextApplyAdditiveTraitSharedData>::Make(DefaultAdditiveTraitData), INDEX_NONE, false, false);
+
+		// Set Input Pins so links to this trait stack can be generated 
+		ChildAInputPin = BlendTraitStackNode->FindTrait(BlendBaseTraitName, GET_MEMBER_NAME_STRING_CHECKED(FAnimNextApplyAdditiveTraitSharedData, Base));
+		ChildBInputPin = BlendTraitStackNode->FindTrait(BlendBaseTraitName, GET_MEMBER_NAME_STRING_CHECKED(FAnimNextApplyAdditiveTraitSharedData, Additive));
+		
+	}
+	else if (BlendMode == EUAFLayerBlendMode::Blend)
+	{
+		if (bUsesPerBoneData)
+		{
+			// Setup default trait values 
+			FAnimNextBlendLayerTraitSharedData DefaultTraitData; 
+			DefaultTraitData.BlendWeight = LayerWeight;
+			DefaultTraitData.BlendProfile = BlendMask;
+			
+			// Create trait 
+			BlendBaseTraitName = LayerCreationContext.GraphController->AddTraitStruct(BlendTraitStackNode, TInstancedStruct<FAnimNextBlendLayerTraitSharedData>::Make(DefaultTraitData), INDEX_NONE, false, false);
+
+			// Set Input Pins so links to this trait stack can be generated 
+			ChildAInputPin = BlendTraitStackNode->FindTrait(BlendBaseTraitName, GET_MEMBER_NAME_STRING_CHECKED(FAnimNextBlendLayerTraitSharedData, ChildBase));
+			ChildBInputPin = BlendTraitStackNode->FindTrait(BlendBaseTraitName, GET_MEMBER_NAME_STRING_CHECKED(FAnimNextBlendLayerTraitSharedData, ChildBlend));
+		}
+		else
+		{
+			// Setup default trait values 
+			FAnimNextBlendTwoWayTraitSharedData DefaultTraitData; 
+			DefaultTraitData.BlendWeight = LayerWeight;
+			
+			// Create two-way blend trait 
+			BlendBaseTraitName = LayerCreationContext.GraphController->AddTraitStruct(BlendTraitStackNode, TInstancedStruct<FAnimNextBlendTwoWayTraitSharedData>::Make(DefaultTraitData), INDEX_NONE, false, false);
+
+			// Set Input Pins so links to this trait stack can be generated 
+			ChildAInputPin = BlendTraitStackNode->FindTrait(BlendBaseTraitName, GET_MEMBER_NAME_STRING_CHECKED(FAnimNextBlendTwoWayTraitSharedData, ChildA));
+			ChildBInputPin = BlendTraitStackNode->FindTrait(BlendBaseTraitName, GET_MEMBER_NAME_STRING_CHECKED(FAnimNextBlendTwoWayTraitSharedData, ChildB));
+		}
+	}
+	else if (BlendMode == EUAFLayerBlendMode::CacheOnly)
+	{
+		// Create passthrough trait on the blend trait stack - since we won't be performing a blend when Cache Only is set  
+		BlendBaseTraitName = LayerCreationContext.GraphController->AddTraitStruct(BlendTraitStackNode, TInstancedStruct<FAnimNextPassthroughSharedData>::Make(), INDEX_NONE, false, false);
+		
+		ChildAInputPin = BlendTraitStackNode->FindTrait(BlendBaseTraitName, GET_MEMBER_NAME_STRING_CHECKED(FAnimNextPassthroughSharedData, Input));
+	}
+
+	if (BlendBaseTraitName == NAME_None)
+	{
+		LayerCreationContext.CompileSettings.ReportError(TEXT("Failed to create the base trait for the blend trait stack!"));
+		return nullptr;
+	}
+	
+	// Expand trait pin per default 
+	if (const URigVMPin* TraitPin = BlendTraitStackNode->FindTrait(BlendBaseTraitName))
+	{
+		LayerCreationContext.GraphController->SetPinExpansion(TraitPin->GetPinPath(), true, false);
+	}
+	
+	// Create the layer data provider 
+	{
+		// Build runtime data on trait 
+		UUAFLayerStack_EditorData* EditorData = UE::UAF::UncookedOnly::FUtils::GetEditorData<UUAFLayerStack_EditorData>(LayerCreationContext.LayerStack.Get());
+		const int32 LayerIndex = EditorData->GetIndexForLayer(LayerCreationContext.Layer);
+		
+		// TODO: This data, which contains the default/bound properties of a layer will likely not live on the trait itself like this in the future once bindings are implemented
+		FUAFLayerProperties LayerProperties; 
+		LayerProperties.bLayerEnabled = bLayerEnabled;
+		LayerProperties.DesiredLayerWeight = LayerWeight;
+		LayerProperties.BlendInTime = LayerBlendInTime;
+		LayerProperties.BlendOutTime = LayerBlendOutTime;
+		LayerProperties.BlendCurve = BlendCurve;
+		LayerProperties.BlendOption = BlendOption;
+		
+		// Setup default trait values 
+		FUAFLayerDataProviderTraitSharedData DefaultTraitData; 
+		DefaultTraitData.LayerName = LayerCreationContext.Layer->GetLayerName();
+		DefaultTraitData.LayerIndex = LayerIndex;
+		DefaultTraitData.bCreateCacheInput = BlendMode == EUAFLayerBlendMode::CacheOnly;
+		DefaultTraitData.LayerStackPath = FSoftObjectPath::ConstructFromObject(LayerCreationContext.LayerStack.Get()).ToString();
+		DefaultTraitData.LayerProperties = MoveTemp(LayerProperties);
+		
+		const FName LayerDataProviderTraitName = LayerCreationContext.GraphController->AddTraitStruct(BlendTraitStackNode, TInstancedStruct<FUAFLayerDataProviderTraitSharedData>::Make(DefaultTraitData), INDEX_NONE, false, false);
+		if (LayerDataProviderTraitName == NAME_None)
+		{
+			LayerCreationContext.CompileSettings.ReportError(TEXT("Could not create the LayerDataProviderTrait."));
+			return nullptr;
+		}
+		
+		// Expand trait pin per default 
+		const URigVMPin* TraitPin = BlendTraitStackNode->FindTrait(LayerDataProviderTraitName);
+		LayerCreationContext.GraphController->SetPinExpansion(TraitPin->GetPinPath(), true, false);
+		
+		// If we are in cache only mode, hook up the layer content output pin to that handle 
+		if (BlendMode == EUAFLayerBlendMode::CacheOnly)
+		{
+			ChildBInputPin = BlendTraitStackNode->FindTrait(LayerDataProviderTraitName, GET_MEMBER_NAME_STRING_CHECKED(FUAFLayerDataProviderTraitSharedData, CacheOnlyInput));
+		}
+	}
+	
+	// Link input pins to blend trait 
+	// Previous Layer Output --> Child A 
+	const bool bLayerBaseInputLinkCreated = LayerCreationContext.GraphController->AddLink(LayerCreationContext.LayerInputs[0], ChildAInputPin, false);
+	// Layer Output --> Child B
+	const bool bLayerContentLinkCreated = LayerCreationContext.GraphController->AddLink(LayerCreationContext.LayerInputs[1], ChildBInputPin, false);
+
+	if (!bLayerBaseInputLinkCreated || !bLayerContentLinkCreated)
+	{
+		LayerCreationContext.CompileSettings.ReportError(TEXT("Failed to create links to blend trait!"));
+		return nullptr;
+	}
+	
+	URigVMPin* BlendTraitStackResult = BlendTraitStackNode->FindPin(GET_MEMBER_NAME_STRING_CHECKED(FRigUnit_AnimNextTraitStack, Result));
+	if (BlendTraitStackResult == nullptr)
+	{
+		LayerCreationContext.CompileSettings.ReportError(TEXT("Failed to retrieve Result pin from trait stack"));
+		return nullptr;
+	}
+	
+	return BlendTraitStackResult;
+}
+
+
+TSharedRef<SWidget> FUAFDefaultBlendProvider::CreateLayerBlendWidget(UUAFLayer* InLayer)
+{
+	return SNew(UE::UAF::Layering::SDefaultBlendProvider)
+		.Layer(InLayer);
+}
+
+const FSlateBrush* FUAFDefaultBlendProvider::GetOverrideLayerBackground() const
+{
+	switch (BlendMode)
+	{
+	case EUAFLayerBlendMode::CacheOnly:
+		return FUAFLayeringStyle::Get().GetBrush(TEXT("Layer.CacheOnlyBackground"));
+	default:
+		return nullptr;
+	}
+}
+
+bool FUAFDefaultBlendProvider::GetOverrideIndicatorColor(FSlateColor& OutSlateColor) const
+{
+	switch (BlendMode)
+	{
+	case EUAFLayerBlendMode::Blend:
+		OutSlateColor = FUAFLayeringStyle::Get().GetSlateColor("Layer.Colors.BlendIndictor");
+		return true;
+	case EUAFLayerBlendMode::Additive:
+		OutSlateColor = FUAFLayeringStyle::Get().GetSlateColor("Layer.Colors.AdditiveIndictor");
+		return true;
+	case EUAFLayerBlendMode::CacheOnly:
+		OutSlateColor = FUAFLayeringStyle::Get().GetSlateColor("Layer.Colors.CacheOnlyIndictor");
+		return true;
+	case EUAFLayerBlendMode::MAX:
+	default:
+		return false;
+	}
+}
+
+namespace UE::UAF::Layering
+{
+bool FDefaultLayerBlendStructureDataProvider::IsValid() const
+{
+	return BlendProvider != nullptr && BlendProvider->IsValid();
+}
+
+const UStruct* FDefaultLayerBlendStructureDataProvider::GetBaseStructure() const
+{
+	return FUAFDefaultBlendProvider::StaticStruct();
+}
+
+void FDefaultLayerBlendStructureDataProvider::GetInstances(TArray<TSharedPtr<FStructOnScope>>& OutInstances, const UStruct* ExpectedBaseStructure) const
+{
+	if (IsValid())
+	{
+		const bool bCompatibleStructType = BlendProvider->GetScriptStruct() == ExpectedBaseStructure;
+		if (bCompatibleStructType)
+		{
+			uint8* Memory = BlendProvider->GetMutableMemory();
+			OutInstances.Add(MakeShareable(new FStructOnScope(BlendProvider->GetScriptStruct(), Memory)));
+		}
+	}
+}
+}
+
+#undef LOCTEXT_NAMESPACE //UAFDefaultBlendProvider
diff --git a/Engine/Plugins/Experimental/UAF/UAFLayering/Source/UAFLayeringUncookedOnly/Private/UAFLayerStack_EditorData.cpp b/Engine/Plugins/Experimental/UAF/UAFLayering/Source/UAFLayeringUncookedOnly/Private/UAFLayerStack_EditorData.cpp
new file mode 100644
--- /dev/null
+++ b/Engine/Plugins/Experimental/UAF/UAFLayering/Source/UAFLayeringUncookedOnly/Private/UAFLayerStack_EditorData.cpp
@@ -0,0 +1,287 @@
+// Copyright Epic Games, Inc. All Rights Reserved.
+
+#pragma once
+
+#include "UAFLayerStack_EditorData.h"
+#include "AnimGraphUncookedOnlyUtils.h"
+#include "UAFLayerStack.h"
+#include "Compilation/AnimNextGetGraphCompileContext.h"
+#include "Entries/AnimNextAnimationGraphEntry.h"
+#include "Entries/AnimNextSharedVariablesEntry.h"
+#include "Entries/AnimNextVariableEntry.h"
+#include "Graph/AnimNextAnimationGraphSchema.h"
+#include "Graph/RigUnit_AnimNextGraphRoot.h"
+#include "Layers/UAFBaseLayer.h"
+#include "Layers/UAFLayer.h"
+#include "Layers/UAFLayerAssetProvider.h"
+#include "Layers/UAFLayerDefaultBlendProvider.h"
+#include "Workspace/UAFLayerStackWorkspaceAssetUserData.h"
+
+void UUAFLayerStack_EditorData::OnPreCompileGetProgrammaticGraphs(const FRigVMCompileSettings& InSettings, FAnimNextGetGraphCompileContext& OutCompileContext)
+{
+	Super::OnPreCompileGetProgrammaticGraphs(InSettings, OutCompileContext);
+	
+	if (UUAFLayerStack* LayerStack = UE::UAF::UncookedOnly::FUtils::GetAsset<UUAFLayerStack>(this))
+	{
+		// Create procedural RigVM Graph to host our layer content and blend logic 
+		URigVMGraph* Graph = NewObject<URigVMGraph>(this, NAME_None, RF_Transient);
+		Graph->SetSchemaClass(UAnimNextAnimationGraphSchema::StaticClass());
+
+		// Create Controller to manipulate the created graph
+		UAnimNextController* Controller = CastChecked<UAnimNextController>(RigVMClient.GetOrCreateController(Graph));
+
+		// Setup basic graph so we can add our custom trait stacks
+		UE::UAF::UncookedOnly::FAnimGraphUtils::SetupAnimGraph(FRigUnit_AnimNextGraphRoot::DefaultEntryPoint, Controller, false);
+
+		if (Controller->GetGraph()->GetNodes().Num() != 1)
+		{
+			InSettings.ReportError(TEXT("Failed to setup UAFAnimGraph - Expected singular FRigUnit_AnimNextGraphRoot node"));
+			return;
+		}
+
+		URigVMNode* EntryNode = Controller->GetGraph()->GetNodes()[0];
+		URigVMPin* GraphResultPin = EntryNode->FindPin(GET_MEMBER_NAME_STRING_CHECKED(FRigUnit_AnimNextGraphRoot, Result));
+		
+		UE::UAF::Layering::FLayerCreationContext LayerCreationContext(InSettings);
+		LayerCreationContext.LayerStack = LayerStack; 
+		LayerCreationContext.LastNodeLocation = FVector2D::ZeroVector;
+		LayerCreationContext.GraphController = Controller;
+
+		URigVMPin* LayerOutputPin = nullptr;
+		const TArray<TObjectPtr<UUAFLayer>> AllLayers = GetAllLayers();
+		for (UUAFLayer* LayerToCreate : AllLayers)
+		{
+			if (!LayerToCreate || LayerToCreate->GetLayerState() == EUAFLayerState::Disabled)
+			{
+				continue;
+			}
+			
+			LayerOutputPin = LayerToCreate->CreateLayerGraphContentTraits(LayerCreationContext);
+			if (LayerOutputPin == nullptr)
+			{
+				InSettings.ReportErrorf(TEXT("Failed to create Layer %s"), *LayerToCreate->GetLayerName().ToString());
+				return;
+			}
+			
+			LayerCreationContext.LayerInputs[0] = LayerOutputPin;
+		}
+
+		// adjust location of results node 
+		LayerCreationContext.LastNodeLocation += FVector2D(500.0f, 0.0f);
+		Controller->SetNodePosition(EntryNode, LayerCreationContext.LastNodeLocation, false);
+
+		// Link the final output to the graph result 
+		const bool bLinkedResultNode = Controller->AddLink(LayerOutputPin, GraphResultPin, false);
+		if (!bLinkedResultNode)
+		{
+			InSettings.ReportError(TEXT("Failed to link final Layer and Graph Output pin"));
+			return;
+		}
+
+		CreatedGraph = Graph;
+		OutCompileContext.GetMutableProgrammaticGraphs().Add(Graph);
+	}
+}
+
+void UUAFLayerStack_EditorData::PostEditChangeProperty(FPropertyChangedEvent& PropertyChangedEvent)
+{
+	Super::PostEditChangeProperty(PropertyChangedEvent);
+
+	if (PropertyChangedEvent.MemberProperty && PropertyChangedEvent.MemberProperty->GetName() == GET_MEMBER_NAME_STRING_CHECKED(UUAFLayerStack_EditorData, Layers))
+	{
+		NotifyLayerLayoutChanged();
+	}
+}
+
+void UUAFLayerStack_EditorData::NotifyLayerLayoutChanged()
+{
+	OnLayerLayoutChanged.Broadcast();
+}
+
+TArray<TObjectPtr<UUAFLayer>> UUAFLayerStack_EditorData::GetAllLayers() const
+{
+	TArray<TObjectPtr<UUAFLayer>> AllLayers;
+	AllLayers.Add(BaseLayer);
+	AllLayers.Append(Layers);
+	
+	return MoveTemp(AllLayers);
+}
+
+void UUAFLayerStack_EditorData::MoveLayerUp(UUAFLayer* LayerToMove)
+{
+	const int32 SourceLayerIndex = Layers.Find(LayerToMove);
+	if (SourceLayerIndex == INDEX_NONE)
+	{
+		// error, not a valid layer 
+		return;
+	}
+
+	if (LayerToMove == BaseLayer)
+	{
+		// can not move the base layer 
+		return;
+	}
+
+	if (SourceLayerIndex == 0)
+	{
+		// can not move this layer up and it is already the uppermost layer after the base layer
+		return;
+	}
+
+	const int32 TargetIndex = SourceLayerIndex - 1;
+	Layers.Swap(SourceLayerIndex, TargetIndex);
+
+	NotifyLayerLayoutChanged();
+	BroadcastModified(EAnimNextEditorDataNotifType::PropertyChanged, this);
+}
+
+void UUAFLayerStack_EditorData::MoveLayerDown(UUAFLayer* LayerToMove)
+{
+	const int32 SourceLayerIndex = Layers.Find(LayerToMove);
+	if (SourceLayerIndex == INDEX_NONE)
+	{
+		// error, not a valid layer 
+		return;
+	}
+
+	if (LayerToMove == BaseLayer)
+	{
+		// cant move base layer 
+		return;
+	}
+	
+	if (SourceLayerIndex == Layers.Num() -1)
+	{
+		// if layer is the last one it makes no sense to move it down. 
+		return;
+	}
+
+	const int32 TargetIndex = SourceLayerIndex + 1;
+	Layers.Swap(SourceLayerIndex, TargetIndex);
+
+	NotifyLayerLayoutChanged();
+	BroadcastModified(EAnimNextEditorDataNotifType::PropertyChanged, this);
+}
+
+void UUAFLayerStack_EditorData::DeleteLayer(UUAFLayer* LayerToDelete)
+{
+	if (LayerToDelete)
+	{
+		if (LayerToDelete == BaseLayer)
+		{
+			// not allowed to remove the base layer
+			return;
+		}
+		Layers.Remove(LayerToDelete);
+		
+		NotifyLayerLayoutChanged();
+		BroadcastModified(EAnimNextEditorDataNotifType::PropertyChanged, this);
+	}
+}
+
+void UUAFLayerStack_EditorData::SetLayerState(UUAFLayer* Layer, EUAFLayerState LayerState)
+{
+	Layer->SetLayerState(LayerState);
+	BroadcastModified(EAnimNextEditorDataNotifType::PropertyChanged, this);
+}
+
+int32 UUAFLayerStack_EditorData::GetIndexForLayer(UUAFLayer* Layer) const 
+{
+	return GetAllLayers().Find(Layer);
+}
+
+int32 UUAFLayerStack_EditorData::GetNumLayers() const
+{
+	return GetAllLayers().Num();
+}
+
+bool UUAFLayerStack_EditorData::IsBaseLayer(UUAFLayer* Layer) const
+{
+	const int32 LayerIndex = GetIndexForLayer(Layer);
+	return LayerIndex == 0;
+}
+
+bool UUAFLayerStack_EditorData::IsLastLayer(UUAFLayer* Layer) const
+{
+	const int32 LayerIndex = GetIndexForLayer(Layer);
+	return LayerIndex == GetAllLayers().Num() - 1;
+}
+
+bool UUAFLayerStack_EditorData::IsLayerNameValid(const FName NewLayerName) const
+{
+	// Empty layer names are not allowed 
+	if (NewLayerName == NAME_None)
+	{
+		return false;
+	}
+	
+	// Layer names need to be unique within the layer stack 
+	for (const UUAFLayer* Layer : GetAllLayers())
+	{
+		if (!Layer)
+		{
+			continue;
+		}
+		
+		if (Layer->GetLayerName() == NewLayerName)
+		{
+			return false;
+		}
+	}
+	return true;
+}
+
+UUAFLayer* UUAFLayerStack_EditorData::CreateDefaultAssetBasedLayer(const FAssetData& InAsset)
+{
+	UUAFLayer* NewLayer = NewObject<UUAFLayer>(this);
+	if (!NewLayer)
+	{
+		UE_LOG(LogAnimation, Error, TEXT("UUAFLayerStack_EditorData::CreateDefaultAssetBasedLayer: Failed to create a new layer!"));
+		return nullptr;
+	}
+	
+	const int32 NewLayerIndex = Layers.Add(NewLayer);
+	if (NewLayerIndex == INDEX_NONE)
+	{
+		UE_LOG(LogAnimation, Error, TEXT(" UUAFLayerStack_EditorData::CreateDefaultAssetBasedLayer: Failed to add new layer to layer array"));
+		return nullptr;
+	}
+	
+	FUAFLayerAssetProvider NewContentProvider;
+	NewContentProvider.SetLayerAsset(InAsset.GetAsset());
+	
+	NewLayer->SetLayerContentProvider(TInstancedStruct<FUAFLayerAssetProvider>::Make(NewContentProvider));
+	NewLayer->SetLayerBlendProvider(TInstancedStruct<FUAFDefaultBlendProvider>::Make());
+	
+	NotifyLayerLayoutChanged();
+	BroadcastModified(EAnimNextEditorDataNotifType::PropertyChanged, this);
+	
+	return NewLayer;
+}
+
+void UUAFLayerStack_EditorData::Initialize(bool bRecompileVM)
+{
+	Super::Initialize(bRecompileVM);
+	
+	if (BaseLayer == nullptr)
+	{
+		BaseLayer = NewObject<UUAFBaseLayer>(this);
+	}
+}
+
+TSubclassOf<UAssetUserData> UUAFLayerStack_EditorData::GetAssetUserDataClass() const
+{
+	return UUAFLayerStackWorkspaceAssetUserData::StaticClass();
+}
+
+TConstArrayView<TSubclassOf<UAnimNextRigVMAssetEntry>> UUAFLayerStack_EditorData::GetEntryClasses() const
+{
+	static const TSubclassOf<UAnimNextRigVMAssetEntry> Classes[] =
+	{
+			UAnimNextVariableEntry::StaticClass(),
+			UAnimNextAnimationGraphEntry::StaticClass(),
+			UAnimNextSharedVariablesEntry::StaticClass(),
+	};
+
+	return Classes;
+}
diff --git a/Engine/Plugins/Experimental/UAF/UAFLayering/Source/UAFLayeringUncookedOnly/Private/UAFLayeringStyle.cpp b/Engine/Plugins/Experimental/UAF/UAFLayering/Source/UAFLayeringUncookedOnly/Private/UAFLayeringStyle.cpp
new file mode 100644
--- /dev/null
+++ b/Engine/Plugins/Experimental/UAF/UAFLayering/Source/UAFLayeringUncookedOnly/Private/UAFLayeringStyle.cpp
@@ -0,0 +1,91 @@
+// Copyright Epic Games, Inc. All Rights Reserved.
+
+#include "UAFLayeringStyle.h"
+#include "Styling/SlateStyleRegistry.h"
+#include "Brushes/SlateRoundedBoxBrush.h"
+#include "Styling/StyleColors.h"
+#include "Styling/SlateStyleMacros.h"
+#include "Styling/SlateTypes.h"
+
+FUAFLayeringStyle& FUAFLayeringStyle::Get()
+{
+	static FUAFLayeringStyle Instance;
+	return Instance;
+}
+
+FUAFLayeringStyle::FUAFLayeringStyle()
+	: FSlateStyleSet(TEXT("UAFLayeringStyle"))
+{
+	SetContentRoot(FPaths::EngineContentDir() / TEXT("Editor/Slate"));
+	SetCoreContentRoot(FPaths::EngineContentDir() / TEXT("Slate"));
+
+	Set("Layer.Colors.BlendIndictor", FSlateColor(FColor(112, 214, 255)));
+	Set("Layer.Colors.AdditiveIndictor", FSlateColor(FColor(255, 112, 166)));
+	Set("Layer.Colors.CacheOnlyIndictor", FSlateColor(FLinearColor(0.1f, 0.1f, 0.1f)));
+	
+	Set("LayerStack.OuterBorder", new FSlateRoundedBoxBrush(FStyleColors::Panel.GetSpecifiedColor(), 25.0f));
+	Set("LayerStack.InnerBorder", new FSlateRoundedBoxBrush(FStyleColors::Background.GetSpecifiedColor(), 10.0f));
+	Set("Layer.Background", new FSlateRoundedBoxBrush(FStyleColors::Panel.GetSpecifiedColor(), 10.0f));
+	Set("Layer.DisabledBackground", new FSlateRoundedBoxBrush(FStyleColors::Recessed.GetSpecifiedColor(), 10.0f));
+	
+	Set("Layer.WhiteBackground", new FSlateRoundedBoxBrush(FLinearColor::White, 5.0f));
+	Set("Layer.AdvancedSettingsBackground", new FSlateRoundedBoxBrush(FStyleColors::Recessed.GetSpecifiedColor(), FVector4(0.0f, 0.0f, 10.0f, 10.0f))); //. X = Top Left, Y = Top Right, Z = Bottom Right, W = Bottom Left
+	Set("Layer.SelectedLayer", new FSlateRoundedBoxBrush(FLinearColor(0.5f, 0.3f, 0.8f), 5.0f));
+	
+	Set("LayerStack.MiddleBackground", new FSlateRoundedBoxBrush(FStyleColors::AccentGray.GetSpecifiedColor(), 10.0f));
+	
+	Set("LayerStack.ListView", FTableViewStyle()
+	.SetBackgroundBrush(FSlateNoResource()));
+
+	Set("AddLayerButton", FButtonStyle()
+	.SetNormal(FSlateNoResource())
+	.SetHovered(FSlateRoundedBoxBrush(FLinearColor(1, 1, 1, 0.15f), 4.0f))
+	.SetPressed(FSlateRoundedBoxBrush(FLinearColor(1, 1, 1, 0.25f), 4.0f))
+	.SetNormalPadding(FMargin(0, 0, 0, 1))
+	.SetPressedPadding(FMargin(0, 1, 0, 0)));
+	
+	
+	Set("LayerStack.ListViewRow", FTableRowStyle()
+		.SetEvenRowBackgroundBrush(FSlateNoResource())
+		.SetEvenRowBackgroundHoveredBrush(FSlateNoResource())
+		.SetOddRowBackgroundBrush(FSlateNoResource())
+		.SetOddRowBackgroundHoveredBrush(FSlateNoResource())
+		.SetActiveHighlightedBrush(FSlateNoResource())
+		.SetActiveBrush(FSlateNoResource())
+		.SetActiveHoveredBrush(FSlateNoResource())
+		.SetSelectorFocusedBrush(FSlateRoundedBoxBrush(FStyleColors::AccentBlue.GetSpecifiedColor(), 10.0f))
+		.SetInactiveHighlightedBrush(FSlateNoResource())
+		.SetInactiveBrush(FSlateRoundedBoxBrush(FStyleColors::AccentBlue.GetSpecifiedColor(), 10.0f))
+		.SetInactiveHoveredBrush(FSlateRoundedBoxBrush(FStyleColors::AccentBlue.GetSpecifiedColor(), 10.0f))
+	);
+	
+	const FVector2D Icon16x16(16.0f, 16.0f);
+	const FVector2D Icon64x64(64.0f, 64.0f);
+	const FVector2D Icon32x32(32.0f, 32.0f);
+	
+	Set("ClassIcon.UAFLayerStack", new IMAGE_BRUSH("Icons/layers_16x", Icon16x16));
+	
+	Set("Layer.VisibleLayerIcon", new CORE_IMAGE_BRUSH_SVG("Starship/Common/visible", Icon16x16));
+	Set("Layer.DisabledLayerIcon", new IMAGE_BRUSH("Icons/SourceControlOff_16x", Icon16x16));
+	Set("Layer.NoPreviewLayerIcon", new CORE_IMAGE_BRUSH_SVG("Starship/Common/hidden", Icon16x16));
+
+	Set("Layer.ExpandAdvancedSettingsIcon", new IMAGE_BRUSH("Common/DownArrow", Icon16x16));
+	Set("Layer.CloseAdvancedSettingsIcon", new IMAGE_BRUSH("Common/UpArrow", Icon16x16));
+
+	Set("Layer.AddLayerIcon", new IMAGE_BRUSH("Icons/icon_levels_addlayer_16x", Icon16x16));
+	
+	Set("Layer.PoseIconBig", new IMAGE_BRUSH_SVG("Starship/AnimationBlueprintEditor/AnimationGraphPose", Icon32x32));
+	Set("Layer.PoseIconSmall", new IMAGE_BRUSH_SVG("Starship/AnimationBlueprintEditor/AnimationGraphPose", Icon16x16));
+	
+	Set("Layer.LODThreshold", new IMAGE_BRUSH("DynamicMaterial/Icons/EditorIcons/MaskToggle_16x", Icon16x16));
+	
+	Set("Layer.CacheOnlyBackground", new IMAGE_BRUSH("/Graph/GraphPanel_StripesBackground", Icon64x64, FLinearColor(0.12f, 0.12f, 0.12f, 0.3f), ESlateBrushTileType::Both ));
+	Set("Layer.AddLayerBackground", new BORDER_BRUSH("/Icons/UMG/ToggleOutlines",  FMargin(4.0f/16.0f), FLinearColor::White));
+	
+	FSlateStyleRegistry::RegisterSlateStyle(*this);
+}
+
+FUAFLayeringStyle::~FUAFLayeringStyle()
+{
+	FSlateStyleRegistry::UnRegisterSlateStyle(*this);
+}
diff --git a/Engine/Plugins/Experimental/UAF/UAFLayering/Source/UAFLayeringUncookedOnly/Private/UAFLayeringUncookedOnlyModule.cpp b/Engine/Plugins/Experimental/UAF/UAFLayering/Source/UAFLayeringUncookedOnly/Private/UAFLayeringUncookedOnlyModule.cpp
new file mode 100644
--- /dev/null
+++ b/Engine/Plugins/Experimental/UAF/UAFLayering/Source/UAFLayeringUncookedOnly/Private/UAFLayeringUncookedOnlyModule.cpp
@@ -0,0 +1,21 @@
+// Copyright Epic Games, Inc. All Rights Reserved.
+
+#include "UAFLayeringUncookedOnlyModule.h"
+
+#define LOCTEXT_NAMESPACE "FUAFLayeringEditorModule"
+
+namespace UE::UAF::Layering
+{
+void FUAFLayeringUncookedOnlyModule::StartupModule()
+{
+}
+
+void FUAFLayeringUncookedOnlyModule::ShutdownModule()
+{
+}
+
+IMPLEMENT_MODULE(FUAFLayeringUncookedOnlyModule, UAFLayeringUncookedOnly)
+}
+
+#undef LOCTEXT_NAMESPACE
+	
diff --git a/Engine/Plugins/Experimental/UAF/UAFLayering/Source/UAFLayeringUncookedOnly/Private/UAFLayeringUncookedOnlyModule.h b/Engine/Plugins/Experimental/UAF/UAFLayering/Source/UAFLayeringUncookedOnly/Private/UAFLayeringUncookedOnlyModule.h
new file mode 100644
--- /dev/null
+++ b/Engine/Plugins/Experimental/UAF/UAFLayering/Source/UAFLayeringUncookedOnly/Private/UAFLayeringUncookedOnlyModule.h
@@ -0,0 +1,18 @@
+// Copyright Epic Games, Inc. All Rights Reserved.
+
+#pragma once
+
+#include "Modules/ModuleManager.h"
+
+namespace UE::UAF::Layering
+{
+
+class FUAFLayeringUncookedOnlyModule : public IModuleInterface
+{
+public:
+	/** IModuleInterface implementation */
+	virtual void StartupModule() override;
+	virtual void ShutdownModule() override;
+};
+
+}
diff --git a/Engine/Plugins/Experimental/UAF/UAFLayering/Source/UAFLayeringUncookedOnly/Private/Widgets/SDefaultBlendProvider.cpp b/Engine/Plugins/Experimental/UAF/UAFLayering/Source/UAFLayeringUncookedOnly/Private/Widgets/SDefaultBlendProvider.cpp
new file mode 100644
--- /dev/null
+++ b/Engine/Plugins/Experimental/UAF/UAFLayering/Source/UAFLayeringUncookedOnly/Private/Widgets/SDefaultBlendProvider.cpp
@@ -0,0 +1,171 @@
+// Copyright Epic Games, Inc. All Rights Reserved.
+
+#include "Widgets/SDefaultBlendProvider.h"
+
+#include "ISinglePropertyView.h"
+#include "SlateOptMacros.h"
+#include "Layers/UAFLayer.h"
+#include "Layers/UAFLayerDefaultBlendProvider.h"
+#include "Widgets/Input/SEditableText.h"
+#include "Widgets/Input/SEditableTextBox.h"
+#include "Widgets/Layout/SBox.h"
+#include "Widgets/Layout/SSplitter.h"
+#include "Widgets/Text/STextBlock.h"
+
+#define LOCTEXT_NAMESPACE "UAFDefaultBlendProvider"
+
+namespace UE::UAF::Layering
+{
+void SDefaultBlendProvider::Construct(const FArguments& InArgs)
+{
+	WeakLayer = InArgs._Layer;
+	if (!ensureMsgf(WeakLayer.IsValid(), TEXT("Tried to create a SDefaultBlendProvider without a valid Layer")))
+	{
+		return;
+	}
+	
+	FSinglePropertyParams DefaultPropertyArgs;
+	DefaultPropertyArgs.NamePlacement = EPropertyNamePlacement::Hidden;
+	DefaultPropertyArgs.bHideResetToDefault = true;
+	DefaultPropertyArgs.bHideAssetThumbnail = true;
+	
+	TSharedRef<FDefaultLayerBlendStructureDataProvider> LayerBlendDataProvider = MakeShared<FDefaultLayerBlendStructureDataProvider>(&WeakLayer->GetLayerBlendProvider());
+	
+	// Generate single property views to display blend properties within the layer widget 
+	TSharedPtr<ISinglePropertyView> LayerBlendModePropertyView = GenerateSinglePropertyView(GET_MEMBER_NAME_CHECKED(FUAFDefaultBlendProvider, BlendMode), LayerBlendDataProvider, DefaultPropertyArgs);
+	TSharedPtr<ISinglePropertyView> LayerWeightPropertyView = GenerateSinglePropertyView(GET_MEMBER_NAME_CHECKED(FUAFDefaultBlendProvider, LayerWeight), LayerBlendDataProvider, DefaultPropertyArgs);
+
+	TSharedPtr<ISinglePropertyView> LayerMaskPropertyView = GenerateSinglePropertyView(GET_MEMBER_NAME_CHECKED(FUAFDefaultBlendProvider, BlendMask), LayerBlendDataProvider, DefaultPropertyArgs);
+	TSharedPtr<ISinglePropertyView> BlendCurvePropertyView = GenerateSinglePropertyView(GET_MEMBER_NAME_CHECKED(FUAFDefaultBlendProvider, BlendCurve), LayerBlendDataProvider, DefaultPropertyArgs);
+
+	// Display the name for the following property views
+	DefaultPropertyArgs.NamePlacement = EPropertyNamePlacement::Left;
+	TSharedPtr<ISinglePropertyView> LayerEnabledPropertyView = GenerateSinglePropertyView(GET_MEMBER_NAME_CHECKED(FUAFDefaultBlendProvider, bLayerEnabled), LayerBlendDataProvider, DefaultPropertyArgs);
+	TSharedPtr<ISinglePropertyView> LayerBlendInTimePropertyView = GenerateSinglePropertyView(GET_MEMBER_NAME_CHECKED(FUAFDefaultBlendProvider, LayerBlendInTime), LayerBlendDataProvider, DefaultPropertyArgs);
+	TSharedPtr<ISinglePropertyView> LayerBlendOutTimePropertyView = GenerateSinglePropertyView(GET_MEMBER_NAME_CHECKED(FUAFDefaultBlendProvider, LayerBlendOutTime), LayerBlendDataProvider, DefaultPropertyArgs);
+
+	ChildSlot
+	[
+		SNew(SSplitter)
+		.Orientation(Orient_Horizontal)
+		.PhysicalSplitterHandleSize(4.0f)
+		.HitDetectionSplitterHandleSize(6.0f)
+		+ SSplitter::Slot()
+		.Value(0.3f)
+		[
+			SNew(SBox)
+			.Padding(10.0f)
+			[
+				SNew(SVerticalBox)
+				+ SVerticalBox::Slot()
+				.Padding(5.0f, 2.0f)
+				[
+					SNew(STextBlock)
+					.Text(LOCTEXT("BlendModeLabel", "Blend Mode"))
+				]
+				+ SVerticalBox::Slot()
+				.Padding(0.0f, 2.0f)
+				[
+					LayerBlendModePropertyView.ToSharedRef()
+				]
+				+ SVerticalBox::Slot()
+				.Padding(0.0f, 2.0f)
+				[
+					LayerMaskPropertyView.ToSharedRef()
+				]
+			]
+
+		]
+		+ SSplitter::Slot()
+		.Value(0.3f)
+		[
+			SNew(SBox)
+			.Padding(10.0f)
+			[
+				SNew(SVerticalBox)
+				+ SVerticalBox::Slot()
+				.Padding(5.0f, 2.0f)
+				[
+					SNew(STextBlock)
+					.Text(LOCTEXT("WeightLabel", "Weight"))
+				]
+				+ SVerticalBox::Slot()
+				.Padding(0.0f, 2.0f)
+				[
+					LayerWeightPropertyView.ToSharedRef()
+				]
+				+ SVerticalBox::Slot()
+				.Padding(5.0f, 2.0f)
+				[
+					SNew(SEditableTextBox)
+					.Text(FText::FromString(TEXT("Layer Debug Weight")))
+					.IsEnabled(false)
+				]
+			]
+
+		]
+		+ SSplitter::Slot()
+		.Value(0.4f)
+		[
+			SNew(SBox)
+			.Padding(10.0f)
+			[
+				SNew(SVerticalBox)
+				+ SVerticalBox::Slot()
+				.Padding(5.0f, 2.0f)
+				[
+					SNew(STextBlock)
+					.Text(LOCTEXT("BlendSettingLabel", "Blend Settings"))
+				]
+				+ SVerticalBox::Slot()
+				.Padding(0.0f, 2.0f)
+				[
+					SNew(SHorizontalBox)
+					+ SHorizontalBox::Slot()
+					[
+						LayerEnabledPropertyView.ToSharedRef()
+					]
+					+ SHorizontalBox::Slot()
+					[
+						BlendCurvePropertyView.ToSharedRef()
+					]
+
+				]
+				+ SVerticalBox::Slot()
+				.Padding(0.0f, 2.0f)
+				[
+					SNew(SHorizontalBox)
+					+ SHorizontalBox::Slot()
+					[
+						LayerBlendInTimePropertyView.ToSharedRef()
+					]
+					+ SHorizontalBox::Slot()
+					[
+						LayerBlendOutTimePropertyView.ToSharedRef()
+					]
+
+				]
+			]
+
+		]
+	];
+}
+
+TSharedPtr<ISinglePropertyView> SDefaultBlendProvider::GenerateSinglePropertyView(const FName PropertyName, const TSharedRef<FDefaultLayerBlendStructureDataProvider> StructDataProvider, const FSinglePropertyParams& Params) const
+{
+	FPropertyEditorModule& PropertyEditorModule = FModuleManager::Get().LoadModuleChecked<FPropertyEditorModule>("PropertyEditor");
+
+	TSharedPtr<ISinglePropertyView> GeneratedPropertyView = PropertyEditorModule.CreateSingleProperty(StructDataProvider.ToSharedPtr(), PropertyName, Params);
+	GeneratedPropertyView->SetOnPropertyValueChanged(FSimpleDelegate::CreateLambda([this]
+		{
+			if (WeakLayer.IsValid())
+			{
+				WeakLayer->BroadcastModified(EAnimNextEditorDataNotifType::PropertyChanged);
+			}
+		}));
+
+	return GeneratedPropertyView;
+}
+}
+
+#undef LOCTEXT_NAMESPACE //UAFDefaultBlendProvider
\ No newline at end of file
diff --git a/Engine/Plugins/Experimental/UAF/UAFLayering/Source/UAFLayeringUncookedOnly/Private/Widgets/SLayer.cpp b/Engine/Plugins/Experimental/UAF/UAFLayering/Source/UAFLayeringUncookedOnly/Private/Widgets/SLayer.cpp
new file mode 100644
--- /dev/null
+++ b/Engine/Plugins/Experimental/UAF/UAFLayering/Source/UAFLayeringUncookedOnly/Private/Widgets/SLayer.cpp
@@ -0,0 +1,483 @@
+// Copyright Epic Games, Inc. All Rights Reserved.
+
+
+#include "Widgets/SLayer.h"
+
+#include "SlateOptMacros.h"
+#include "UAFLayeringStyle.h"
+#include "Components/VerticalBox.h"
+#include "Layers/UAFLayer.h"
+#include "Widgets/SOverlay.h"
+#include "Widgets/Images/SImage.h"
+#include "Widgets/Input/SButton.h"
+#include "Widgets/Input/SEditableTextBox.h"
+#include "Widgets/Layout/SBorder.h"
+#include "Widgets/Layout/SSplitter.h"
+#include "Widgets/Text/SInlineEditableTextBlock.h"
+#include "Widgets/Text/STextBlock.h"
+#include "Widgets/Views/STableRow.h"
+
+#define LOCTEXT_NAMESPACE "FUAFLayerWidget"
+
+namespace UE::UAF::Layering
+{
+	
+void SLayer::Construct(const FArguments& InArgs)
+{
+	Layer = InArgs._Layer;
+	ensure(Layer.IsValid());
+
+	ChildSlot
+	[
+		SNew(SOverlay)
+		+ SOverlay::Slot()
+		.HAlign(HAlign_Fill)
+		.VAlign(VAlign_Fill)
+		[
+			SNew(SBorder)
+			.BorderImage(this, &SLayer::GetBackgroundBrush)
+			[
+				SNew(SBorder)
+				.BorderImage(this, &SLayer::GetBackgroundOverlayBrush)
+				[
+					SNew(SVerticalBox)
+					+ SVerticalBox::Slot()
+					[
+						SNew(SHorizontalBox)
+						+ SHorizontalBox::Slot()
+						.MaxWidth(5.0f)
+						.Padding(10.0f)
+						[
+							SNew(SBorder)
+							.BorderImage(FUAFLayeringStyle::Get().GetBrush(TEXT("Layer.WhiteBackground")))
+							.BorderBackgroundColor(this, &SLayer::GetIndicatorColor)
+						]
+						+ SHorizontalBox::Slot()
+						.AutoWidth()
+						[
+							// Grabber Button for drag and drop
+							SNew(SButton)
+							.ContentPadding(FMargin(5.0f, 5.0f, 5.0f, 5.0f))
+							.ButtonStyle(FAppStyle::Get(), "NoBorder")
+							.ForegroundColor(FSlateColor::UseForeground())
+							.HAlign(HAlign_Center)
+							.VAlign(VAlign_Center)
+							.Content()
+							[
+								SNew(SImage)
+								.Image(FAppStyle::GetBrush("VerticalBoxDragIndicatorShort"))
+							]
+						]
+						+ SHorizontalBox::Slot()
+						.AutoWidth()
+						.HAlign(HAlign_Center)
+						.VAlign(VAlign_Center)
+						[
+							// Visibility button 
+							SNew(SButton)
+							.ContentPadding(FMargin(5.0f))
+							.ButtonStyle(FAppStyle::Get(), "NoBorder")
+							.ForegroundColor(FSlateColor::UseForeground())
+							.HAlign(HAlign_Center)
+							.VAlign(VAlign_Center)
+							.ToolTipText_Lambda([this]() -> FText
+								{
+									if (Layer.IsValid())
+									{
+										switch (Layer->GetLayerState())
+										{
+										case EUAFLayerState::Enabled:
+											return LOCTEXT("LayerVisibilityTooltip_DisablePreview", "Disable this layer for preview. It will still be enabled at runtime.");
+										case EUAFLayerState::PreviewDisabled:
+											return LOCTEXT("LayerVisibilityTooltip_EnablePreview", "Enable this layer for preview.");
+										case EUAFLayerState::Disabled:
+											return LOCTEXT("LayerVisibilityTooltip_Disabled", "This layer is currently disabled for preview and runtime.");
+										default:
+											return FText();
+										}
+									}
+									return FText();
+								})
+							.OnClicked_Lambda([this]() -> FReply
+								{
+									if (Layer.IsValid())
+									{
+										Layer->TogglePreviewVisibility();
+									}
+									return FReply::Handled();
+								})
+							.Content()
+							[
+								SNew(SImage)
+								.Image(this, &SLayer::GetVisibilityBrushForLayer)
+							]
+						]
+						+ SHorizontalBox::Slot()
+						.Padding(5.0f, 10.f)
+						.FillContentWidth(1.0f)
+						[
+							SNew(SSplitter)
+							.Orientation(Orient_Horizontal)
+							.PhysicalSplitterHandleSize(4.0f)
+							.HitDetectionSplitterHandleSize(6.0f)
+							+ SSplitter::Slot()
+							.Value(0.2)
+							[
+								SNew(SVerticalBox)
+								+ SVerticalBox::Slot()
+								.AutoHeight()
+								.Padding(10.0f, 10.0f, 0.0f, 0.0f)
+								[
+									SNew(SInlineEditableTextBlock)
+									.OnTextCommitted(this, &SLayer::OnLayerNameCommitted)
+									.ToolTipText(this, &SLayer::GetLayerNameTooltip)
+									.Text(this, &SLayer::GetLayerName)
+									.IsSelected(this, &SLayer::IsSelected)
+									.OnVerifyTextChanged(this, &SLayer::VerifyLayerName)
+								]
+								+ SVerticalBox::Slot()
+								.AutoHeight()
+								.Padding(10.0f, 5.0f)
+								[
+									Layer.IsValid() && Layer->IsLayerContentProviderValid()
+									? Layer->GetLayerContentProvider().GetMutable().CreateLayerContentWidget(Layer.Get())
+									: SNullWidget::NullWidget
+								]
+								+ SVerticalBox::Slot()
+								.AutoHeight()
+								.Padding(10.0f, 0.0f)
+								[
+									SNew(SHorizontalBox)
+									+ SHorizontalBox::Slot()
+									.AutoWidth()
+									[
+										SNew(SImage)
+										.Image(FUAFLayeringStyle::Get().GetBrush(TEXT("Layer.PoseIconSmall")))
+										.ColorAndOpacity(FLinearColor::Yellow)
+									]
+
+								]
+
+							]
+							+ SSplitter::Slot()
+							.Value(0.8)
+							[
+								Layer.IsValid() && Layer->IsLayerBlendProviderValid()
+								? Layer->GetLayerBlendProvider().GetMutable().CreateLayerBlendWidget(Layer.Get())
+								: SNullWidget::NullWidget
+							]
+						]
+						+ SHorizontalBox::Slot()
+						.AutoWidth()
+						.HAlign(HAlign_Right)
+						.Padding(20.0f, 0.0f, 0.0f, 0.0f)
+						[
+							// Advanced Settings button 
+							SNew(SButton)
+							.ContentPadding(FMargin(5.0f, 5.0f, 20.0f, 5.0f))
+							.ButtonStyle(FAppStyle::Get(), "NoBorder")
+							.ForegroundColor(FSlateColor::UseForeground())
+							.HAlign(HAlign_Center)
+							.VAlign(VAlign_Center)
+							.OnClicked_Lambda([this]() -> FReply
+								{
+									bAdvancedSettingsExpanded = !bAdvancedSettingsExpanded;
+									return FReply::Handled();
+								})
+							.Content()
+							[
+								SNew(SImage)
+								.Image(this, &SLayer::GetAdvancedSettingBrushForLayer)
+							]
+						]
+
+					]
+					+ SVerticalBox::Slot() // Advanced Settings Container 
+					.AutoHeight()
+					[
+						SNew(SBorder)
+						.BorderImage(FUAFLayeringStyle::Get().GetBrush(TEXT("Layer.AdvancedSettingsBackground")))
+						.Padding(10.0f)
+						.Visibility_Lambda([this]() -> EVisibility
+							{
+								return bAdvancedSettingsExpanded
+									? EVisibility::Visible
+									: EVisibility::Collapsed;
+							})
+						[
+							SNew(SSplitter)
+							.Orientation(Orient_Horizontal)
+							.PhysicalSplitterHandleSize(4.0f)
+							.HitDetectionSplitterHandleSize(6.0f)
+							+ SSplitter::Slot() // Cached Output Pose
+							.Value(0.2)
+							[
+								SNew(SVerticalBox)
+								+ SVerticalBox::Slot()
+								.Padding(20.0f, 10.0f)
+								.AutoHeight()
+								[
+									SNew(STextBlock)
+									.Text(LOCTEXT("OutputPoseLabel", "Output Pose"))
+								]
+								+ SVerticalBox::Slot()
+								.Padding(20.0f, 0.0f)
+								.AutoHeight()
+								[
+									SNew(SHorizontalBox)
+									+ SHorizontalBox::Slot()
+									.Padding(0.0f, 2.0f)
+									.AutoWidth()
+									[
+										SNew(SImage)
+										.Image(FUAFLayeringStyle::Get().GetBrush(TEXT("Layer.PoseIconSmall")))
+										.ColorAndOpacity(FLinearColor::Yellow)
+									]
+									+ SHorizontalBox::Slot()
+									.AutoWidth()
+									.Padding(5.0f, 0.0f)
+									[
+										SNew(SEditableTextBox)
+										.Text(FText::FromString(TEXT("Output Pose Binding")))
+										.IsEnabled(false)
+									]
+								]
+							]
+							+ SSplitter::Slot() // Cached Additive Pose
+							.Value(0.4f)
+							[
+								SNew(SHorizontalBox) 
+								+ SHorizontalBox::Slot()
+								.Padding(20.0f, 10.0f, 5.0f, 0.0f)
+								.AutoWidth()
+								.HAlign(HAlign_Left)
+								[
+									SNew(SVerticalBox)
+									+ SVerticalBox::Slot()
+									.HAlign(HAlign_Right)
+									.AutoHeight()
+									.Padding(0.0f, 0.0f, 0.0f, 5.0f)
+									[
+										SNew(SHorizontalBox)
+										+ SHorizontalBox::Slot()
+										.AutoWidth()
+										[
+											SNew(STextBlock)
+											.Text(LOCTEXT("AdditiveOutputPoseBindingLabel", "Output Additive Pose"))
+										]
+										+ SHorizontalBox::Slot()
+										.AutoWidth()
+										.Padding(5.0f, 0.0f)
+										[
+											SNew(SImage)
+											.Image(FUAFLayeringStyle::Get().GetBrush(TEXT("Layer.PoseIconSmall"))) // TODO: Change to additive icon
+											.ColorAndOpacity(FLinearColor::Yellow)
+										]
+									]
+									+ SVerticalBox::Slot()
+									.Padding(0.0f, 5.0f, 0.0f, 5.0f)
+									.HAlign(HAlign_Right)
+									[
+										SNew(STextBlock)
+										.Text(LOCTEXT("AdditiveBaseLabel", "Base"))
+									]
+								]
+								+ SHorizontalBox::Slot()
+								.HAlign(HAlign_Left)
+								.AutoWidth()
+								[
+									SNew(SVerticalBox)
+									+ SVerticalBox::Slot()
+									.AutoHeight()
+									.Padding(5.0f)
+									[
+										// Additive binding
+										SNew(SEditableTextBox)
+										.Text(FText::FromString(TEXT("Additive Pose Binding")))
+										.IsEnabled(false)
+									]
+									+ SVerticalBox::Slot()
+									.Padding(5.0f)
+									.AutoHeight()
+									[
+										SNew(SHorizontalBox)
+										+ SHorizontalBox::Slot()
+										.AutoWidth()
+										[
+											SNew(SEditableTextBox)
+											.Text(FText::FromString(TEXT("Base For Additive")))
+											.IsEnabled(false)
+										]
+										+ SHorizontalBox::Slot()
+										.AutoWidth()
+										.Padding(5.0f, 0.0f, 5.0f, 0.0f)
+										[
+											SNew(STextBlock)
+											.Text(LOCTEXT("AdditiveLabel", "Additive"))
+										]
+										+ SHorizontalBox::Slot()
+										.AutoWidth()
+										[
+											SNew(SEditableTextBox)
+											.Text(FText::FromString(TEXT("Additive For Additive")))
+											.IsEnabled(false)
+										]
+									]
+								]
+							]
+							+ SSplitter::Slot() // LOD Threshold
+							.Value(0.4f)
+							[
+								SNew(SVerticalBox)
+								+ SVerticalBox::Slot()
+								.Padding(20.0f, 5.0f)
+								.AutoHeight()
+								[
+									SNew(STextBlock)
+									.Text(LOCTEXT("LODSettingLabel", "LOD Settings"))
+								]
+								+ SVerticalBox::Slot()
+								.Padding(20.0f, 5.0f)
+								.AutoHeight()
+								[
+									SNew(SHorizontalBox)
+									+ SHorizontalBox::Slot()
+									.AutoWidth()
+									[
+										// Icon
+										SNew(SImage)
+										.Image(FUAFLayeringStyle::Get().GetBrush(TEXT("Layer.LODThreshold")))
+									]
+									+ SHorizontalBox::Slot()
+									.Padding(5.0f, 0.0f)
+									.AutoWidth()
+									[
+										// Binding
+										SNew(SEditableTextBox)
+										.Text(FText::FromString(TEXT("Mask Binding")))
+										.IsEnabled(false)
+									]
+								]
+							]
+						]
+					]
+				]
+			]
+		]
+	];
+
+}
+
+const FSlateBrush* SLayer::GetVisibilityBrushForLayer() const
+{
+	if (Layer.IsValid())
+	{
+		switch (Layer->GetLayerState())
+		{
+		case EUAFLayerState::Enabled:
+			return FUAFLayeringStyle::Get().GetBrush(TEXT("Layer.VisibleLayerIcon"));
+		case EUAFLayerState::PreviewDisabled: 
+			return FUAFLayeringStyle::Get().GetBrush(TEXT("Layer.NoPreviewLayerIcon"));
+		case EUAFLayerState::Disabled:
+			return FUAFLayeringStyle::Get().GetBrush(TEXT("Layer.DisabledLayerIcon"));
+		default:
+			return FUAFLayeringStyle::Get().GetBrush(TEXT("Layer.VisibleLayerIcon"));
+		}
+	}
+	return FUAFLayeringStyle::Get().GetBrush(TEXT("Layer.VisibleLayerIcon"));
+}
+
+const FSlateBrush* SLayer::GetAdvancedSettingBrushForLayer() const
+{
+	if (bAdvancedSettingsExpanded)
+	{
+		return FUAFLayeringStyle::Get().GetBrush(TEXT("Layer.CloseAdvancedSettingsIcon"));
+	}
+
+	return FUAFLayeringStyle::Get().GetBrush(TEXT("Layer.ExpandAdvancedSettingsIcon"));
+}
+
+const FSlateBrush* SLayer::GetBackgroundBrush() const
+{
+	return  Layer.IsValid() && Layer->GetLayerState() == EUAFLayerState::Disabled ? FUAFLayeringStyle::Get().GetBrush(TEXT("Layer.DisabledBackground")) : FUAFLayeringStyle::Get().GetBrush(TEXT("Layer.Background"));
+}
+
+const FSlateBrush* SLayer::GetBackgroundOverlayBrush() const
+{
+	// Ask providers for overrides 
+	if (Layer.IsValid() && Layer->IsLayerBlendProviderValid())
+	{
+		if (const FSlateBrush* OverrideBrush = Layer->GetLayerBlendProvider().Get().GetOverrideLayerBackground())
+		{
+			return OverrideBrush;
+		}
+	}
+
+	return FAppStyle::GetNoBrush();
+}
+
+void SLayer::OnLayerNameCommitted(const FText& NewText, ETextCommit::Type CommitInfo) const
+{
+	if (Layer.IsValid())
+	{
+		Layer->RenameLayer(FName(NewText.ToString()));
+	}
+}
+
+FSlateColor SLayer::GetIndicatorColor() const
+{
+	if (Layer.IsValid() && Layer->IsLayerBlendProviderValid())
+	{
+		FSlateColor OverrideColor;
+		if (Layer->GetLayerBlendProvider().Get().GetOverrideIndicatorColor(OverrideColor))
+		{
+			return OverrideColor;
+		}
+	}
+
+	return FSlateColor(FLinearColor::White);
+}
+
+FText SLayer::GetLayerNameTooltip() const
+{
+	return LOCTEXT("LayerNameTooltip", "The name of this layer. The name has to be unique within the layer stack");
+}
+
+FText SLayer::GetLayerName() const
+{
+	if (Layer.IsValid())
+	{
+		return FText::FromName(Layer->GetLayerName());
+	}
+	return FText();
+}
+
+bool SLayer::IsSelected() const
+{
+	if (TSharedPtr<STableRow<TObjectPtr<UUAFLayer>>> OwningRow = StaticCastSharedPtr<STableRow<TObjectPtr<UUAFLayer>>>(this->GetParentWidget()))
+	{
+		return OwningRow->IsSelected();
+	}
+
+	return false;
+}
+
+bool SLayer::VerifyLayerName(const FText& Text, FText& OutErrorMessage) const
+{
+	if (Layer.IsValid())
+	{
+		if (const bool bValidName = Layer->IsLayerNameValid(FName(Text.ToString())))
+		{
+			return true;
+		}
+		
+		OutErrorMessage = FText::FromString(TEXT("Invalid layer name. Layer names have to be unique within the layer stack and not none."));
+	}
+	
+	return false;
+}
+	
+}
+
+#undef LOCTEXT_NAMESPACE
+
diff --git a/Engine/Plugins/Experimental/UAF/UAFLayering/Source/UAFLayeringUncookedOnly/Private/Widgets/SLayerStack.cpp b/Engine/Plugins/Experimental/UAF/UAFLayering/Source/UAFLayeringUncookedOnly/Private/Widgets/SLayerStack.cpp
new file mode 100644
--- /dev/null
+++ b/Engine/Plugins/Experimental/UAF/UAFLayering/Source/UAFLayeringUncookedOnly/Private/Widgets/SLayerStack.cpp
@@ -0,0 +1,386 @@
+// Copyright Epic Games, Inc. All Rights Reserved.
+
+#include "Widgets/SLayerStack.h"
+
+#include "AnimNextAnimGraphSettings.h"
+#include "UAFLayerStack.h"
+#include "IWorkspaceEditor.h"
+#include "UncookedOnlyUtils.h"
+#include "UAFLayerStack_EditorData.h"
+#include "Layers/UAFLayer.h"
+#include "Widgets/Views/SListView.h"
+#include "Widgets/SOverlay.h"
+#include "Widgets/SBoxPanel.h"
+#include "UAFLayeringStyle.h"
+#include "SAssetDropTarget.h"
+#include "LayeringUncookedOnlyTypes.h"
+#include "Layers/UAFBaseLayer.h"
+
+#define LOCTEXT_NAMESPACE "FUAFLayerStackWidget"
+
+namespace UE::UAF::Layering
+{
+
+void SLayerStack::Construct(const FArguments& InArgs)
+{
+	// args
+	LayerStack = InArgs._LayerStack;
+	WeakWorkspaceEditor = InArgs._WorkspaceEditor;
+	ensure(LayerStack.IsValid());
+
+	UUAFLayerStack_EditorData* EditorData = UE::UAF::UncookedOnly::FUtils::GetEditorData<UUAFLayerStack_EditorData>(LayerStack.Get());
+	EditorData->OnLayerLayoutChanged.AddSP(this, &SLayerStack::OnLayerLayoutChanged);
+	CacheLayersForList();
+
+	// make commands for context menu
+	CommandList = MakeShareable(new FUICommandList);
+
+	// widget content 
+	ChildSlot
+	[
+		SNew(SBorder)
+		.Padding(10.0f)
+		.HAlign(HAlign_Fill)
+		.VAlign(VAlign_Fill)
+		.BorderImage(FUAFLayeringStyle::Get().GetBrush(TEXT("LayerStack.OuterBorder")))
+		[
+			SNew(SVerticalBox)
+			+ SVerticalBox::Slot()
+			.FillContentHeight(1.0f)
+			[
+				SNew(SOverlay)
+				+ SOverlay::Slot()
+				.HAlign(HAlign_Fill)
+				.VAlign(VAlign_Fill)
+				[
+					SNew(SBorder)
+					.BorderImage(FUAFLayeringStyle::Get().GetBrush(TEXT("LayerStack.InnerBorder")))
+				]
+				+ SOverlay::Slot()
+				.Padding(10.0f)
+				[
+					SNew(SVerticalBox)
+					+ SVerticalBox::Slot()
+					[
+						SAssignNew(LayerListView, SListView<TWeakObjectPtr<UUAFLayer>>)
+						.ListViewStyle(FUAFLayeringStyle::Get(), "LayerStack.ListView")
+						.ListItemsSource(&CachedLayerPointers)
+						.OnGenerateRow(this, &SLayerStack::GenerateLayerRow)
+						.ClearSelectionOnClick(true)
+						.SelectionMode(ESelectionMode::Single)
+						.OnSelectionChanged(this, &SLayerStack::OnListSelectionChanged)
+						.OnContextMenuOpening(this, &SLayerStack::OnConstructContextMenu)
+					]
+					+ SVerticalBox::Slot()
+					.MinHeight(50.0f)
+					.AutoHeight()
+					.VAlign(VAlign_Top)
+					.Padding(20.0f)
+					[
+						SNew(SAssetDropTarget)
+						.bSupportsMultiDrop(false)
+						.OnAreAssetsAcceptableForDrop_Lambda([](TArrayView<FAssetData> InAssets)
+							{
+								const TArray<UClass*> AllowedClasses = UAnimNextAnimGraphSettings::GetAllowedAssetClasses();
+								for (const FAssetData& AssetData : InAssets)
+								{
+									if (AllowedClasses.Contains(AssetData.GetClass()))
+									{
+										return true;
+									}
+								}
+								return false;
+							})
+						.OnAssetsDropped_Lambda([EditorData](const FDragDropEvent&, TArrayView<FAssetData> InAssets)
+							{
+								const TArray<UClass*> AllowedClasses = UAnimNextAnimGraphSettings::GetAllowedAssetClasses();
+								for (const FAssetData& AssetData : InAssets)
+								{
+									if (AllowedClasses.Contains(AssetData.GetClass()))
+									{
+										EditorData->CreateDefaultAssetBasedLayer(AssetData);
+									}
+								}
+							})
+						[
+							SNew(SBorder)
+							.VAlign(VAlign_Fill)
+							.HAlign(HAlign_Fill)
+							.BorderImage(FUAFLayeringStyle::Get().GetBrush(TEXT("Layer.AddLayerBackground")))
+							[
+								SNew(SButton)
+								.ContentPadding(FMargin(10.0F))
+								.ButtonStyle(FUAFLayeringStyle::Get(), "AddLayerButton")
+								.ForegroundColor(FSlateColor::UseForeground())
+								.HAlign(HAlign_Center)
+								.VAlign(VAlign_Center)
+								.OnReleased_Lambda([EditorData]()
+									{
+										EditorData->CreateDefaultAssetBasedLayer(FAssetData());
+									})
+								.Content()
+								[
+									SNew(SHorizontalBox)
+									+ SHorizontalBox::Slot()
+									.Padding(5.0f)
+									.AutoWidth()
+									[
+										SNew(SImage)
+										.Image(FUAFLayeringStyle::Get().GetBrush(TEXT("Layer.AddLayerIcon")))
+									]
+									+ SHorizontalBox::Slot()
+									.Padding(5.0f)
+									.AutoWidth()
+									[
+										SNew(STextBlock)
+										.Text(LOCTEXT("AddLayerLabel", "Add Layer"))
+									]
+
+								]
+							]
+						]
+					]
+				]
+			]
+			+ SVerticalBox::Slot()
+			.AutoHeight()
+			.VAlign(VAlign_Bottom)
+			.Padding(0.0f, 20.0f)
+			[
+				SNew(SHorizontalBox)
+				+ SHorizontalBox::Slot()
+				.AutoWidth()
+				[
+					SNew(SBorder)
+					.Padding(10.0f)
+					.BorderImage(FUAFLayeringStyle::Get().GetBrush(TEXT("LayerStack.InnerBorder")))
+					[
+						SNew(SHorizontalBox)
+						+ SHorizontalBox::Slot()
+						.AutoWidth()
+						[
+							SNew(SBorder)
+							.VAlign(VAlign_Center)
+							.HAlign(HAlign_Center)
+							.Padding(10.0f)
+							.BorderImage(FUAFLayeringStyle::Get().GetBrush(TEXT("Layer.Background")))
+							[
+								SNew(SImage)
+								.Image(FUAFLayeringStyle::Get().GetBrush(TEXT("Layer.PoseIconBig")))
+							]
+						]
+						+ SHorizontalBox::Slot()
+						.Padding(10.0f, 10.0f)
+						[
+							SNew(SVerticalBox)
+							+ SVerticalBox::Slot()
+							[
+								SNew(STextBlock)
+								.Text(LOCTEXT("FinalCachedPoseLabel", "Final Cached Pose"))
+							]
+						]
+					]
+				]
+				+ SHorizontalBox::Slot()
+				.VAlign(VAlign_Bottom)
+				.HAlign(HAlign_Right)
+				[
+					SNew(SButton)
+					.Text(FText::FromString(TEXT("Show Graph")))
+					.OnClicked_Lambda([EditorData]()
+						{
+							if (EditorData->CreatedGraph.IsValid())
+							{
+								UE::UAF::UncookedOnly::FUtils::OpenProgrammaticGraphs(EditorData, {EditorData->CreatedGraph.Get()});
+							}
+
+							return FReply::Handled();
+						})
+				]
+			]
+		]
+	];
+}
+
+SLayerStack::~SLayerStack()
+{
+	LayerListView->ClearItemsSource();
+	CachedLayerPointers.Empty();
+}
+
+void SLayerStack::OnListSelectionChanged(TWeakObjectPtr<UUAFLayer> SelectedItem, ESelectInfo::Type SelectInfo) const
+{
+	if (TSharedPtr<UE::Workspace::IWorkspaceEditor> SharedWorkspaceEditor = WeakWorkspaceEditor.Pin())
+	{
+		if (SelectedItem.IsValid())
+		{
+			SharedWorkspaceEditor->SetDetailsObjects({ SelectedItem.Get() });
+		}
+		else
+		{
+			SharedWorkspaceEditor->SetDetailsObjects({ LayerStack.Get() });
+			LayerListView->ClearSelection();
+		}
+	}
+}
+
+TSharedRef<ITableRow> SLayerStack::GenerateLayerRow(TWeakObjectPtr<UUAFLayer> Item, const TSharedRef<STableViewBase>& OwnerTable) const
+{
+	return SNew(STableRow<TWeakObjectPtr<UUAFLayer>>, OwnerTable)
+		.Padding(3.0f)
+		.Style(FUAFLayeringStyle::Get(), "LayerStack.ListViewRow")
+		[
+			Item.IsValid() ? Item->CreateLayerWidget() : SNullWidget::NullWidget
+		];
+}
+
+TSharedPtr<SWidget> SLayerStack::OnConstructContextMenu() const
+{
+	TArray<TWeakObjectPtr<UUAFLayer>> SelectedLayers = LayerListView->GetSelectedItems();
+	if (SelectedLayers.Num() == 0 || !SelectedLayers[0].IsValid())
+	{
+		return SNullWidget::NullWidget;
+	}
+	
+	UUAFLayerStack_EditorData* EditorData = UE::UAF::UncookedOnly::FUtils::GetEditorData<UUAFLayerStack_EditorData>(LayerStack.Get());
+	if (!EditorData)
+	{
+		return SNullWidget::NullWidget;
+	}
+	
+	UUAFLayer* SelectedLayer = SelectedLayers[0].Get();
+	const int32 LayerIndex = EditorData->GetIndexForLayer(SelectedLayer);
+	const int32 NumLayers = EditorData->GetNumLayers();
+	const bool bIsBaseLayer = EditorData->IsBaseLayer(SelectedLayer);
+	
+	constexpr bool bShouldCloseWindowAfterMenuSelection = true;
+	FMenuBuilder MenuBuilder(bShouldCloseWindowAfterMenuSelection, CommandList);
+	MenuBuilder.BeginSection("LayerStackLayout", LOCTEXT("LayerStackLayoutLabel", "Layer Stack Layout"));
+	{
+		const bool bCanMoveLayerUp = LayerIndex > 1;
+		const bool bCanMoveLayerDown = LayerIndex > 0 && LayerIndex < NumLayers - 1;
+
+		if (bCanMoveLayerUp)
+		{
+			MenuBuilder.AddMenuEntry(
+			LOCTEXT("MoveLayerUpLabel", "Move Layer Up"),
+			LOCTEXT("MoveLayerUpTooltip", "Moves layer up in the execution order."),
+			FSlateIcon(),
+			FUIAction(FExecuteAction::CreateLambda([EditorData, SelectedLayer, this]()
+				{
+					EditorData->MoveLayerUp(SelectedLayer);
+				})),
+			NAME_None,
+			EUserInterfaceActionType::Button);
+		}
+		
+		if (bCanMoveLayerDown)
+		{
+			MenuBuilder.AddMenuEntry(
+			LOCTEXT("MoveLayerDownLabel", "Move Layer Down"),
+			LOCTEXT("MoveLayerDownTooltip", "Moves layer down in the execution order."),
+			FSlateIcon(),
+			FUIAction(
+				FExecuteAction::CreateLambda([EditorData, SelectedLayer, this]()
+				{
+					EditorData->MoveLayerDown(SelectedLayer);
+				})),
+			NAME_None,
+			EUserInterfaceActionType::Button);
+		}
+	}
+	MenuBuilder.EndSection();
+
+	MenuBuilder.BeginSection("LayerOperation", LOCTEXT("LayerOperationLabel", "Layer Operations"));
+	{
+		// TODO: to be implemented
+		//MenuBuilder.AddMenuEntry(
+		//	LOCTEXT("RenameLayerLabel", "Rename Layer"),
+		//	LOCTEXT("RenameLayerTooltip", "Renames the selected layer."),
+		//	FSlateIcon(),
+		//	FUIAction(FExecuteAction::CreateLambda([EditorData, SelectedLayers]()
+		//		{
+
+		//		})),
+		//		NAME_None,
+		//		EUserInterfaceActionType::Button);
+	
+		if (!bIsBaseLayer)
+		{
+			if (SelectedLayer->GetLayerState() == EUAFLayerState::Disabled)
+			{
+				MenuBuilder.AddMenuEntry(
+				LOCTEXT("EnableLayerLabel", "Enable Layer"),
+				LOCTEXT("EnableLayerTooltip", "Enables the selected layer."),
+				FSlateIcon(),
+				FUIAction(FExecuteAction::CreateLambda([EditorData, SelectedLayer]()
+					{
+						EditorData->SetLayerState(SelectedLayer, EUAFLayerState::Enabled);
+					})),
+				NAME_None,
+				EUserInterfaceActionType::Button);
+			}
+			else
+			{
+				MenuBuilder.AddMenuEntry(
+				LOCTEXT("DisableLayerLabel", "Disable Layer"),
+				LOCTEXT("DisableLayerTooltip", "Disables the selected layer."),
+				FSlateIcon(),
+				FUIAction(FExecuteAction::CreateLambda([EditorData, SelectedLayer]()
+					{
+						EditorData->SetLayerState(SelectedLayer, EUAFLayerState::Disabled);
+					})),
+				NAME_None,
+				EUserInterfaceActionType::Button);
+			}
+		
+			MenuBuilder.AddMenuEntry(
+				LOCTEXT("DeleteLayerLabel", "Delete Layer"),
+				LOCTEXT("DeleteLayerTooltip", "Deletes the selected layer."),
+				FSlateIcon(),
+				FUIAction(FExecuteAction::CreateLambda([EditorData, SelectedLayer]()
+					{
+						EditorData->DeleteLayer(SelectedLayer);
+					})),
+				NAME_None,
+				EUserInterfaceActionType::Button);
+		
+		}
+	}
+	MenuBuilder.EndSection();
+
+	return MenuBuilder.MakeWidget();
+}
+
+void SLayerStack::OnLayerLayoutChanged()
+{
+	CacheLayersForList();
+	LayerListView->RebuildList();
+}
+	
+void SLayerStack::CacheLayersForList()
+{
+	if (!LayerStack.IsValid())
+	{
+		CachedLayerPointers.Empty();
+		return;
+	}
+		
+	UUAFLayerStack_EditorData* EditorData = UE::UAF::UncookedOnly::FUtils::GetEditorData<UUAFLayerStack_EditorData>(LayerStack.Get());
+	if (!EditorData)
+	{
+		CachedLayerPointers.Empty();
+		return;
+	}
+		
+	const TArray<TObjectPtr<UUAFLayer>> Layers = EditorData->GetAllLayers();
+	CachedLayerPointers.Reset(Layers.Num());
+	for (const TObjectPtr<UUAFLayer>& Layer : Layers)
+	{
+		CachedLayerPointers.Add(Layer);
+	}
+}
+	
+}
+#undef LOCTEXT_NAMESPACE
+	
diff --git a/Engine/Plugins/Experimental/UAF/UAFLayering/Source/UAFLayeringUncookedOnly/Private/Workspace/UAFLayerStackWorkspaceAssetUserData.cpp b/Engine/Plugins/Experimental/UAF/UAFLayering/Source/UAFLayeringUncookedOnly/Private/Workspace/UAFLayerStackWorkspaceAssetUserData.cpp
new file mode 100644
--- /dev/null
+++ b/Engine/Plugins/Experimental/UAF/UAFLayering/Source/UAFLayeringUncookedOnly/Private/Workspace/UAFLayerStackWorkspaceAssetUserData.cpp
@@ -0,0 +1,16 @@
+// Copyright Epic Games, Inc. All Rights Reserved.
+
+#pragma once
+
+#include "Workspace/UAFLayerStackWorkspaceAssetUserData.h"
+#include "UAFLayerStack.h"
+
+void UUAFLayerStackWorkspaceAssetUserData::GetRootAssetExport(FAssetRegistryTagsContext Context) const
+{
+	Super::GetRootAssetExport(Context);
+}
+
+void UUAFLayerStackWorkspaceAssetUserData::GetWorkspaceAssetExports(FAssetRegistryTagsContext Context) const
+{
+	Super::GetWorkspaceAssetExports(Context);
+}
diff --git a/Engine/Plugins/Experimental/UAF/UAFLayering/Source/UAFLayeringUncookedOnly/Public/LayeringUncookedOnlyTypes.h b/Engine/Plugins/Experimental/UAF/UAFLayering/Source/UAFLayeringUncookedOnly/Public/LayeringUncookedOnlyTypes.h
new file mode 100644
--- /dev/null
+++ b/Engine/Plugins/Experimental/UAF/UAFLayering/Source/UAFLayeringUncookedOnly/Public/LayeringUncookedOnlyTypes.h
@@ -0,0 +1,60 @@
+// Copyright Epic Games, Inc. All Rights Reserved.
+
+#pragma once
+
+#include "RigVMCompiler/RigVMCompiler.h"
+#include "LayeringUncookedOnlyTypes.generated.h"
+
+class UAnimNextController;
+class UUAFLayer;
+class UUAFLayerStack;
+
+UENUM(BlueprintType)
+enum class EUAFLayerBlendMode : uint8
+{
+	Blend		UMETA(DisplayName= "Blend", ToolTip = "Behaves like a regular blend and will blend together based on layer weight. Weight 0 == Full Source Value, 1 == Full Layer Value"),
+	Additive	UMETA(DisplayName = "Additive", ToolTip = "The layer will get applied as a additive with the weight of the layer"),
+	CacheOnly	UMETA(DisplayName = "Cache Only", ToolTip = "The layer will not get applied at all, but will produce a cached output pose if bound"),
+	MAX			UMETA(Hidden)
+};
+
+UENUM(BlueprintType)
+enum class EUAFLayerState : uint8
+{
+	Enabled			UMETA(DisplayName="Enabled", ToolTip = "This layer is enabled during preview and runtime."),
+	PreviewDisabled	UMETA(DisplayName="Preview Disabled", ToolTip = "This layer is disabled during preview but will execute during runtime."),
+	Disabled		UMETA(DisplayName="Disabled", ToolTip = "This layer is disabled during preview and runtime.")
+};
+
+namespace UE::UAF::Layering
+{
+	struct FLayerCreationContext
+	{
+		FLayerCreationContext(const FRigVMCompileSettings& InCompileSettings)
+			: CompileSettings(InCompileSettings)
+		{
+			LayerInputs.SetNum(2);
+		}
+	
+		// The compile setting, can be used to report errors 
+		const FRigVMCompileSettings& CompileSettings;
+	
+		// The stack this layer creation takes place in 
+		TObjectPtr<UUAFLayerStack> LayerStack = nullptr;
+	
+		// The layer currently being generated 
+		TObjectPtr<UUAFLayer> Layer = nullptr;
+	
+		// The UAF Controller controlling the underlying graph of this layer stack
+		TObjectPtr<UAnimNextController> GraphController = nullptr;
+	
+		// Input Pins 
+		// 0 == Last layers output 
+		// 1 == This layers content output 
+		TArray<URigVMPin*> LayerInputs;
+
+		// The location of the last created node
+		// Used to build a readable graph under the hood with correct node placement 
+		FVector2D LastNodeLocation = FVector2D::ZeroVector;
+	};
+}
diff --git a/Engine/Plugins/Experimental/UAF/UAFLayering/Source/UAFLayeringUncookedOnly/UAFLayeringUncookedOnly.Build.cs b/Engine/Plugins/Experimental/UAF/UAFLayering/Source/UAFLayeringUncookedOnly/UAFLayeringUncookedOnly.Build.cs
new file mode 100644
--- /dev/null
+++ b/Engine/Plugins/Experimental/UAF/UAFLayering/Source/UAFLayeringUncookedOnly/UAFLayeringUncookedOnly.Build.cs
@@ -0,0 +1,62 @@
+// Copyright Epic Games, Inc. All Rights Reserved.
+
+using System.Collections.Generic;
+using UnrealBuildTool;
+using UnrealBuildTool.Rules;
+
+public class UAFLayeringUncookedOnly : ModuleRules
+{
+	public UAFLayeringUncookedOnly(ReadOnlyTargetRules Target) : base(Target)
+	{
+		PCHUsage = ModuleRules.PCHUsageMode.UseExplicitOrSharedPCHs;
+
+		PublicDependencyModuleNames.AddRange(
+			new string[]
+			{
+				"Core",
+				"CoreUObject",
+				"Engine",
+				"Slate",
+				"SlateCore",
+				"WorkspaceEditor",
+				"UAF",
+				"UAFUncookedOnly",
+				"UAFAnimGraphUncookedOnly",
+				"PropertyEditor",
+				"HierarchyTableRuntime",
+				"HierarchyTableAnimationRuntime",
+				"InputCore"
+				
+				// ... add other public dependencies that you statically link with here ...
+			}
+			);
+
+		PrivateDependencyModuleNames.AddRange(
+			new string[]
+			{
+				"UAFLayering",
+				"RigVM",
+				"RigVMDeveloper",
+				"UAFAnimGraph",
+				"UnrealEd",
+				"EditorWidgets",
+				"EditorInteractiveToolsFramework",
+				"InteractiveToolsFramework",
+				// ... add private dependencies that you statically link with here ...	
+			}
+		);
+		
+		if (Target.bBuildEditor)
+		{
+			PrivateIncludePathModuleNames.AddRange(
+				new string[]
+				{
+					"UAFLayeringEditor",
+				}
+			);
+			
+		}
+		
+		bAllowUETypesInNamespaces = true;
+	}
+}
diff --git a/Engine/Plugins/Experimental/UAF/UAFLayering/UAFLayering.uplugin b/Engine/Plugins/Experimental/UAF/UAFLayering/UAFLayering.uplugin
new file mode 100644
--- /dev/null
+++ b/Engine/Plugins/Experimental/UAF/UAFLayering/UAFLayering.uplugin
@@ -0,0 +1,51 @@
+{
+	"FileVersion": 3,
+	"Version": 1,
+	"VersionName": "0.1",
+	"FriendlyName": "UAF Layering",
+	"Description": "Framework to define a layering setup in UAF",
+	"Category": "Animation",
+	"CreatedBy": "Epic Games, Inc.",
+	"CreatedByURL": "https://epicgames.com",
+	"DocsURL": "",
+	"MarketplaceURL": "",
+	"SupportURL": "",
+	"EnabledByDefault": false,
+	"CanContainContent": true,
+	"IsBetaVersion": false,
+	"IsExperimentalVersion": true,
+	"Installed": false,
+	"Plugins": 
+	[	
+		{
+			"Name": "Workspace",
+			"Enabled": true
+		},
+		{
+			"Name": "UAF",
+			"Enabled": true
+		},
+		{
+			"Name": "UAFAnimGraph",
+			"Enabled": true
+		}
+	],
+	"Modules":
+	[
+		{
+			"Name": "UAFLayering",
+			"Type": "Runtime",
+			"LoadingPhase": "Default"
+		},
+		{
+			"Name": "UAFLayeringEditor",
+			"Type": "Editor",
+			"LoadingPhase": "Default"
+		},
+		{
+			"Name": "UAFLayeringUncookedOnly",
+			"Type": "UncookedOnly",
+			"LoadingPhase": "Default"
+		}
+	]
+}
\ No newline at end of file
