UAF: Populate editor mode info for the layering editor mode

#Synced-CL 49839474

diff --git a/Engine/Plugins/Experimental/UAF/UAFLayering/Source/UAFLayeringEditor/Private/UAFLayeringEditorMode.cpp b/Engine/Plugins/Experimental/UAF/UAFLayering/Source/UAFLayeringEditor/Private/UAFLayeringEditorMode.cpp
--- a/Engine/Plugins/Experimental/UAF/UAFLayering/Source/UAFLayeringEditor/Private/UAFLayeringEditorMode.cpp
+++ b/Engine/Plugins/Experimental/UAF/UAFLayering/Source/UAFLayeringEditor/Private/UAFLayeringEditorMode.cpp
@@ -1,6 +1,17 @@
 // Copyright Epic Games, Inc. All Rights Reserved.
 
-
 #include "UAFLayeringEditorMode.h"
 
-const FEditorModeID UUAFLayeringEditorMode::EM_UAFLayeringMode("UAFLayeringEditorMode");
\ No newline at end of file
+#define LOCTEXT_NAMESPACE "UAFLayeringEditorMode"
+
+const FEditorModeID UUAFLayeringEditorMode::EM_UAFLayeringMode("UAFLayeringEditorMode");
+
+UUAFLayeringEditorMode::UUAFLayeringEditorMode()
+{
+	Info = FEditorModeInfo(UUAFLayeringEditorMode::EM_UAFLayeringMode,
+	LOCTEXT("UAFLayeringEditorModeName", "UAFLayeringEditorMode"),
+	FSlateIcon(),
+	false);
+}
+
+#undef LOCTEXT_NAMESPACE // "UAFLayeringEditorMode"
diff --git a/Engine/Plugins/Experimental/UAF/UAFLayering/Source/UAFLayeringEditor/Private/UAFLayeringEditorMode.h b/Engine/Plugins/Experimental/UAF/UAFLayering/Source/UAFLayeringEditor/Private/UAFLayeringEditorMode.h
--- a/Engine/Plugins/Experimental/UAF/UAFLayering/Source/UAFLayeringEditor/Private/UAFLayeringEditorMode.h
+++ b/Engine/Plugins/Experimental/UAF/UAFLayering/Source/UAFLayeringEditor/Private/UAFLayeringEditorMode.h
@@ -13,4 +13,5 @@ class UUAFLayeringEditorMode : public UEdMode
 public:
 	const static FEditorModeID EM_UAFLayeringMode;
 	
+	UUAFLayeringEditorMode();
 };
