UAF: Layering Update

- Add additive per bone blend support
- Fix certain pins not being latent or inline them
- Make the pointer to the generated graph transient

#Synced-CL 50345329

diff --git a/Engine/Plugins/Experimental/UAF/UAFLayering/Source/UAFLayering/Internal/Traits/LayerAssetDataTraitData.h b/Engine/Plugins/Experimental/UAF/UAFLayering/Source/UAFLayering/Internal/Traits/LayerAssetDataTraitData.h
--- a/Engine/Plugins/Experimental/UAF/UAFLayering/Source/UAFLayering/Internal/Traits/LayerAssetDataTraitData.h
+++ b/Engine/Plugins/Experimental/UAF/UAFLayering/Source/UAFLayering/Internal/Traits/LayerAssetDataTraitData.h
@@ -22,7 +22,8 @@ struct FUAFLayerAssetDataTraitSharedData : public FAnimNextTraitSharedData
 	
 	// Latent pin support boilerplate
 	#define TRAIT_LATENT_PROPERTIES_ENUMERATOR(GeneratorMacro) \
-	GeneratorMacro(Asset) 
+	GeneratorMacro(Asset) \
+	GeneratorMacro(Parameters)
 
 	GENERATE_TRAIT_LATENT_PROPERTIES(FUAFLayerAssetDataTraitSharedData, TRAIT_LATENT_PROPERTIES_ENUMERATOR)
 	#undef TRAIT_LATENT_PROPERTIES_ENUMERATOR
diff --git a/Engine/Plugins/Experimental/UAF/UAFLayering/Source/UAFLayering/Internal/Traits/LayerDataProviderTraitData.h b/Engine/Plugins/Experimental/UAF/UAFLayering/Source/UAFLayering/Internal/Traits/LayerDataProviderTraitData.h
--- a/Engine/Plugins/Experimental/UAF/UAFLayering/Source/UAFLayering/Internal/Traits/LayerDataProviderTraitData.h
+++ b/Engine/Plugins/Experimental/UAF/UAFLayering/Source/UAFLayering/Internal/Traits/LayerDataProviderTraitData.h
@@ -59,15 +59,15 @@ struct FUAFLayerDataProviderTraitSharedData : public FAnimNextTraitSharedData
 	FUAFLayerProperties LayerProperties;
 
 	// The index of this layer in its owning layer stack
-	UPROPERTY(EditAnywhere, Category = "Layer")
+	UPROPERTY(EditAnywhere, Category = "Layer", meta = (Inline))
 	int32 LayerIndex = INDEX_NONE;
 
 	// The layer name of this layer 
-	UPROPERTY(EditAnywhere, Category = "Layer")
+	UPROPERTY(EditAnywhere, Category = "Layer", meta = (Inline))
 	FName LayerName = NAME_None;
 	
 	// The soft object path to the owning layer stack
-	UPROPERTY(EditAnywhere, Category = "Layer")
+	UPROPERTY(EditAnywhere, Category = "Layer", meta = (Inline))
 	FString LayerStackPath;
 	
 	// If this is true we will allocate the input to allow for update/eval of cache only branches of the tree
diff --git a/Engine/Plugins/Experimental/UAF/UAFLayering/Source/UAFLayering/Private/Traits/LayerAssetDataProviderTrait.cpp b/Engine/Plugins/Experimental/UAF/UAFLayering/Source/UAFLayering/Private/Traits/LayerAssetDataProviderTrait.cpp
--- a/Engine/Plugins/Experimental/UAF/UAFLayering/Source/UAFLayering/Private/Traits/LayerAssetDataProviderTrait.cpp
+++ b/Engine/Plugins/Experimental/UAF/UAFLayering/Source/UAFLayering/Private/Traits/LayerAssetDataProviderTrait.cpp
@@ -42,13 +42,13 @@ void FLayerAssetDataProviderTrait::PreUpdate(FUpdateTraversalContext& Context, c
 			return;
 		}
 		
-		if (!SharedData->Parameters.IsValid())
+		FAnimNextFactoryParams Params = SharedData->GetParameters(Binding);
+		if (!Params.IsValid())
 		{
 			UAF_TRAIT_LOG(Warning, TEXT("No valid factory params provided. Can not create graph."));
 			return;
 		}
 		
-		FAnimNextFactoryParams Params = SharedData->Parameters;
 		const UUAFAnimGraph* AnimationGraph = IGraphFactory::GetOrBuildGraph(Context, Binding, SharedDataAsset, Params);
 		if(AnimationGraph == nullptr)
 		{
@@ -60,7 +60,7 @@ void FLayerAssetDataProviderTrait::PreUpdate(FUpdateTraversalContext& Context, c
 		NewGraphRequest.BlendArgs = FAlphaBlendArgs(0.0f);
 		NewGraphRequest.FactoryObject = SharedDataAsset;
 		NewGraphRequest.AnimationGraph = AnimationGraph;
-		NewGraphRequest.FactoryParams = SharedData->Parameters;
+		NewGraphRequest.FactoryParams = Params;
 
 		BlendStackBinding.PushGraph(Context, MoveTemp(NewGraphRequest));
 		InstanceData->CachedAsset = SharedDataAsset;
diff --git a/Engine/Plugins/Experimental/UAF/UAFLayering/Source/UAFLayeringUncookedOnly/Internal/UAFLayerStack_EditorData.h b/Engine/Plugins/Experimental/UAF/UAFLayering/Source/UAFLayeringUncookedOnly/Internal/UAFLayerStack_EditorData.h
--- a/Engine/Plugins/Experimental/UAF/UAFLayering/Source/UAFLayeringUncookedOnly/Internal/UAFLayerStack_EditorData.h
+++ b/Engine/Plugins/Experimental/UAF/UAFLayering/Source/UAFLayeringUncookedOnly/Internal/UAFLayerStack_EditorData.h
@@ -48,7 +48,7 @@ protected:
 	// End UObject interface
 	
 public:
-	UPROPERTY()
+	UPROPERTY(Transient)
 	TWeakObjectPtr<URigVMGraph> CreatedGraph = nullptr;
 	
 	FLayerLayoutChangedDelegate OnLayerLayoutChanged;
diff --git a/Engine/Plugins/Experimental/UAF/UAFLayering/Source/UAFLayeringUncookedOnly/Private/Layers/UAFLayerDefaultBlendProvider.cpp b/Engine/Plugins/Experimental/UAF/UAFLayering/Source/UAFLayeringUncookedOnly/Private/Layers/UAFLayerDefaultBlendProvider.cpp
--- a/Engine/Plugins/Experimental/UAF/UAFLayering/Source/UAFLayeringUncookedOnly/Private/Layers/UAFLayerDefaultBlendProvider.cpp
+++ b/Engine/Plugins/Experimental/UAF/UAFLayering/Source/UAFLayeringUncookedOnly/Private/Layers/UAFLayerDefaultBlendProvider.cpp
@@ -14,6 +14,7 @@
 #include "RigVMCore/RigVMStruct.h"
 #include "RigVMModel/Nodes/RigVMUnitNode.h"
 #include "TraitCore/Trait.h"
+#include "Traits/ApplyAdditivePerBoneTraitData.h"
 #include "Traits/ApplyAdditiveTraitData.h"
 #include "Traits/ApplyNamedSetTraitData.h"
 #include "Traits/BlendLayerTraitData.h"
@@ -118,20 +119,29 @@ URigVMPin* FUAFDefaultBlendProvider::CreateBlendGraphTrait(UE::UAF::Layering::FL
 	{
 		if (bUsesPerBoneData)
 		{
-			// Not supported yet!
-			LayerCreationContext.CompileSettings.ReportWarning(TEXT("Per Bone application of additives in not supported yet! Will create a fully body additive trait instead"));
+			UE::UAF::FApplyAdditivePerBoneTraitSharedData DefaultAdditiveTraitData;
+			DefaultAdditiveTraitData.Alpha = LayerWeight;
+			DefaultAdditiveTraitData.BlendProfile = BlendMask;
+			
+			// Create trait 
+			BlendBaseTraitName = LayerCreationContext.GraphController->AddTraitStruct(BlendTraitStackNode, TInstancedStruct<UE::UAF::FApplyAdditivePerBoneTraitSharedData>::Make(DefaultAdditiveTraitData), INDEX_NONE, false, false);
+
+			// Set Input Pins so links to this trait stack can be generated 
+			ChildAInputPin = BlendTraitStackNode->FindTrait(BlendBaseTraitName, GET_MEMBER_NAME_STRING_CHECKED(UE::UAF::FApplyAdditivePerBoneTraitSharedData, Base));
+			ChildBInputPin = BlendTraitStackNode->FindTrait(BlendBaseTraitName, GET_MEMBER_NAME_STRING_CHECKED(UE::UAF::FApplyAdditivePerBoneTraitSharedData, Additive));
 		}
+		else
+		{
+			FAnimNextApplyAdditiveTraitSharedData DefaultAdditiveTraitData;
+			DefaultAdditiveTraitData.Alpha = LayerWeight;
 		
-		FAnimNextApplyAdditiveTraitSharedData DefaultAdditiveTraitData;
-		DefaultAdditiveTraitData.Alpha = LayerWeight;
-		
-		// Create trait 
-		BlendBaseTraitName = LayerCreationContext.GraphController->AddTraitStruct(BlendTraitStackNode, TInstancedStruct<FAnimNextApplyAdditiveTraitSharedData>::Make(DefaultAdditiveTraitData), INDEX_NONE, false, false);
+			// Create trait 
+			BlendBaseTraitName = LayerCreationContext.GraphController->AddTraitStruct(BlendTraitStackNode, TInstancedStruct<FAnimNextApplyAdditiveTraitSharedData>::Make(DefaultAdditiveTraitData), INDEX_NONE, false, false);
 
-		// Set Input Pins so links to this trait stack can be generated 
-		ChildAInputPin = BlendTraitStackNode->FindTrait(BlendBaseTraitName, GET_MEMBER_NAME_STRING_CHECKED(FAnimNextApplyAdditiveTraitSharedData, Base));
-		ChildBInputPin = BlendTraitStackNode->FindTrait(BlendBaseTraitName, GET_MEMBER_NAME_STRING_CHECKED(FAnimNextApplyAdditiveTraitSharedData, Additive));
-		
+			// Set Input Pins so links to this trait stack can be generated 
+			ChildAInputPin = BlendTraitStackNode->FindTrait(BlendBaseTraitName, GET_MEMBER_NAME_STRING_CHECKED(FAnimNextApplyAdditiveTraitSharedData, Base));
+			ChildBInputPin = BlendTraitStackNode->FindTrait(BlendBaseTraitName, GET_MEMBER_NAME_STRING_CHECKED(FAnimNextApplyAdditiveTraitSharedData, Additive));
+		}
 	}
 	else if (BlendMode == EUAFLayerBlendMode::Blend)
 	{
