UAF Layering
- Swap over to use Asset Data and Factories instead of using Parameters directly to create graphs
- Add scoped transactions when changing a layer asset

#Synced-CL 50531632

diff --git a/Engine/Plugins/Experimental/UAF/UAFLayering/Source/UAFLayering/Internal/Traits/LayerAssetDataTraitData.h b/Engine/Plugins/Experimental/UAF/UAFLayering/Source/UAFLayering/Internal/Traits/LayerAssetDataTraitData.h
--- a/Engine/Plugins/Experimental/UAF/UAFLayering/Source/UAFLayering/Internal/Traits/LayerAssetDataTraitData.h
+++ b/Engine/Plugins/Experimental/UAF/UAFLayering/Source/UAFLayering/Internal/Traits/LayerAssetDataTraitData.h
@@ -8,7 +8,7 @@
 #include "LayerAssetDataTraitData.generated.h"
 
 
-/** A trait that creates a anim graph at runtime based on asset and factory params of the layering setup and pushes it onto a blend stack*/
+/** A trait that creates an UAF Graph at runtime based on the given asset in the layering setup and pushes it onto a blend stack*/
 USTRUCT(meta = (DisplayName = "Layer Data Provider", ShowTooltip=true, Hidden))
 struct FUAFLayerAssetDataTraitSharedData : public FAnimNextTraitSharedData
 {
@@ -16,14 +16,10 @@ struct FUAFLayerAssetDataTraitSharedData : public FAnimNextTraitSharedData
 	
 	UPROPERTY(EditAnywhere, Category = "Layer|Content", meta = (GetAllowedClasses = "/Script/UAFAnimGraph.AnimNextAnimGraphSettings:GetAllowedAssetClasses"))
 	TObjectPtr<UObject> Asset = nullptr;
-
-	UPROPERTY(EditAnywhere, Category = "Layer|Content", meta=(FactorySourceProperty = "Asset"))
-	FAnimNextFactoryParams Parameters;
 	
 	// Latent pin support boilerplate
 	#define TRAIT_LATENT_PROPERTIES_ENUMERATOR(GeneratorMacro) \
-	GeneratorMacro(Asset) \
-	GeneratorMacro(Parameters)
+	GeneratorMacro(Asset)
 
 	GENERATE_TRAIT_LATENT_PROPERTIES(FUAFLayerAssetDataTraitSharedData, TRAIT_LATENT_PROPERTIES_ENUMERATOR)
 	#undef TRAIT_LATENT_PROPERTIES_ENUMERATOR
diff --git a/Engine/Plugins/Experimental/UAF/UAFLayering/Source/UAFLayering/Private/Traits/LayerAssetDataProviderTrait.cpp b/Engine/Plugins/Experimental/UAF/UAFLayering/Source/UAFLayering/Private/Traits/LayerAssetDataProviderTrait.cpp
--- a/Engine/Plugins/Experimental/UAF/UAFLayering/Source/UAFLayering/Private/Traits/LayerAssetDataProviderTrait.cpp
+++ b/Engine/Plugins/Experimental/UAF/UAFLayering/Source/UAFLayering/Private/Traits/LayerAssetDataProviderTrait.cpp
@@ -6,6 +6,7 @@
 #include "Factory/AnimGraphFactory.h"
 #include "TraitInterfaces/IBlendStack.h"
 #include "TraitInterfaces/IGraphFactory.h"
+#include "UAF/UAFAssetFactory.h"
 
 namespace UE::UAF
 {
@@ -36,34 +37,35 @@ void FLayerAssetDataProviderTrait::PreUpdate(FUpdateTraversalContext& Context, c
 	if (SharedDataAsset != InstanceData->CachedAsset)
 	{
 		TTraitBinding<IBlendStack> BlendStackBinding;
-		if(!Binding.GetStackInterface<IBlendStack>(BlendStackBinding))
+		if(Binding.GetStackInterface<IBlendStack>(BlendStackBinding))
 		{
-			UAF_TRAIT_LOG(Error, TEXT("Blend Stack Trait is required but not found in stack!"));
-			return;
-		}
+			InstanceData->CachedAsset = SharedDataAsset;
+			if (SharedDataAsset != nullptr)
+			{
+				FGraphAssetHandle AssetHandle = FAssetDataFactory::CreateUAFAssetDataFromObject<FUAFGraphFactoryAsset>(SharedDataAsset);
+				FAnimNextFactoryParams FactoryParams = FAnimGraphFactory::GetDefaultParamsForAsset(AssetHandle);
 		
-		FAnimNextFactoryParams Params = SharedData->GetParameters(Binding);
-		if (!Params.IsValid())
-		{
-			UAF_TRAIT_LOG(Warning, TEXT("No valid factory params provided. Can not create graph."));
-			return;
+				const UUAFAnimGraph* AnimationGraph = IGraphFactory::GetOrBuildGraph(Context, Binding, AssetHandle, FactoryParams);
+				if(AnimationGraph)
+				{
+					IBlendStack::FGraphRequest NewGraphRequest;
+					NewGraphRequest.BlendArgs = FAlphaBlendArgs(0.0f);
+					NewGraphRequest.FactoryObject = SharedDataAsset;
+					NewGraphRequest.AnimationGraph = AnimationGraph;
+					NewGraphRequest.FactoryParams =  MoveTemp(FactoryParams);;
+
+					BlendStackBinding.PushGraph(Context, MoveTemp(NewGraphRequest));
+				}
+				else
+				{
+					UAF_TRAIT_LOG(Error, TEXT("Failed to create graph!"));
+				}
+			}
 		}
-		
-		const UUAFAnimGraph* AnimationGraph = IGraphFactory::GetOrBuildGraph(Context, Binding, SharedDataAsset, Params);
-		if(AnimationGraph == nullptr)
+		else
 		{
-			UAF_TRAIT_LOG(Error, TEXT("Failed to create graph!"));
-			return;
+			UAF_TRAIT_LOG(Error, TEXT("Blend Stack Trait is required but not found in stack!"));
 		}
-
-		IBlendStack::FGraphRequest NewGraphRequest;
-		NewGraphRequest.BlendArgs = FAlphaBlendArgs(0.0f);
-		NewGraphRequest.FactoryObject = SharedDataAsset;
-		NewGraphRequest.AnimationGraph = AnimationGraph;
-		NewGraphRequest.FactoryParams = Params;
-
-		BlendStackBinding.PushGraph(Context, MoveTemp(NewGraphRequest));
-		InstanceData->CachedAsset = SharedDataAsset;
 	}
 		
 	IUpdate::PreUpdate(Context, Binding, TraitState);
diff --git a/Engine/Plugins/Experimental/UAF/UAFLayering/Source/UAFLayeringUncookedOnly/Internal/Layers/UAFLayerAssetProvider.h b/Engine/Plugins/Experimental/UAF/UAFLayering/Source/UAFLayeringUncookedOnly/Internal/Layers/UAFLayerAssetProvider.h
--- a/Engine/Plugins/Experimental/UAF/UAFLayering/Source/UAFLayeringUncookedOnly/Internal/Layers/UAFLayerAssetProvider.h
+++ b/Engine/Plugins/Experimental/UAF/UAFLayering/Source/UAFLayeringUncookedOnly/Internal/Layers/UAFLayerAssetProvider.h
@@ -17,6 +17,17 @@ struct FUAFLayerAssetProvider : public FUAFLayerContentProviderBase
 	GENERATED_BODY()
 
 public: 
+	PRAGMA_DISABLE_DEPRECATION_WARNINGS
+	// Note: We need to explicitly disable warnings on these constructors/operators for clang to be happy with deprecated variables
+	// this is a requirement for clang to compile without warnings.
+	FUAFLayerAssetProvider() = default;
+	virtual ~FUAFLayerAssetProvider() override = default;
+	FUAFLayerAssetProvider(const FUAFLayerAssetProvider&) = default;
+	FUAFLayerAssetProvider(FUAFLayerAssetProvider&&) = default;
+	FUAFLayerAssetProvider& operator=(const FUAFLayerAssetProvider&) = default;
+	FUAFLayerAssetProvider& operator=(FUAFLayerAssetProvider&&) = default;
+	PRAGMA_ENABLE_DEPRECATION_WARNINGS
+	
 	virtual URigVMPin* CreateLayerContentTrait(UE::UAF::Layering::FLayerCreationContext& LayerCreationContext) override;
 	virtual TSharedRef<SWidget> CreateLayerContentWidget(UUAFLayer* InLayer) override;
 
@@ -29,8 +40,11 @@ private:
 	UPROPERTY(EditAnywhere, Category = "Layer|Content", meta = (GetAllowedClasses = "/Script/UAFAnimGraph.AnimNextAnimGraphSettings:GetAllowedAssetClasses"))
 	TObjectPtr<UObject> LayerAsset = nullptr;
 
-	UPROPERTY(EditAnywhere, Category = "Layer|Content", meta=(FactorySourceProperty = "LayerAsset"))
-	FAnimNextFactoryParams Parameters;
+#if WITH_EDITORONLY_DATA
+	UE_DEPRECATED(5.8, "Factory Parameters are not used directly anymore, instead it uses FUAFAssetData internally")
+	UPROPERTY()
+	FAnimNextFactoryParams Parameters_DEPRECATED;
+#endif
 	
 	UPROPERTY()
 	TWeakObjectPtr<UUAFLayer> OuterLayer = nullptr;
diff --git a/Engine/Plugins/Experimental/UAF/UAFLayering/Source/UAFLayeringUncookedOnly/Private/Layers/UAFLayerAssetProvider.cpp b/Engine/Plugins/Experimental/UAF/UAFLayering/Source/UAFLayeringUncookedOnly/Private/Layers/UAFLayerAssetProvider.cpp
--- a/Engine/Plugins/Experimental/UAF/UAFLayering/Source/UAFLayeringUncookedOnly/Private/Layers/UAFLayerAssetProvider.cpp
+++ b/Engine/Plugins/Experimental/UAF/UAFLayering/Source/UAFLayeringUncookedOnly/Private/Layers/UAFLayerAssetProvider.cpp
@@ -8,6 +8,7 @@
 #include "AnimNextRigVMAssetEditorData.h"
 #include "AssetThumbnail.h"
 #include "PropertyCustomizationHelpers.h"
+#include "ScopedTransaction.h"
 #include "Factory/AnimGraphFactory.h"
 #include "Graph/RigUnit_AnimNextTraitStack.h"
 #include "Layers/UAFLayer.h"
@@ -23,7 +24,7 @@
 
 URigVMPin* FUAFLayerAssetProvider::CreateLayerContentTrait(UE::UAF::Layering::FLayerCreationContext& LayerCreationContext)
 {
-	LayerCreationContext.LastNodeLocation += FVector2D(0.0f, 300.0f);
+	LayerCreationContext.LastNodeLocation += FVector2D(0.0f, 500.0f);
 	URigVMUnitNode* ContentHostTraitStackNode = LayerCreationContext.GraphController->AddUnitNode(FRigUnit_AnimNextTraitStack::StaticStruct(), FRigVMStruct::ExecuteName, LayerCreationContext.LastNodeLocation, FString(), false);
 	if (ContentHostTraitStackNode == nullptr)
 	{
@@ -61,7 +62,6 @@ URigVMPin* FUAFLayerAssetProvider::CreateLayerContentTrait(UE::UAF::Layering::FL
 		// Create asset data provider trait 
 		FUAFLayerAssetDataTraitSharedData AssetDataDefaults;
 		AssetDataDefaults.Asset = LayerAsset;
-		AssetDataDefaults.Parameters = Parameters;
 		
 		const FName AssetDataTraitName = LayerCreationContext.GraphController->AddTraitStruct(ContentHostTraitStackNode, TInstancedStruct<FUAFLayerAssetDataTraitSharedData>::Make(AssetDataDefaults), INDEX_NONE, false, false);
 		if (AssetDataTraitName.IsNone())
@@ -69,6 +69,9 @@ URigVMPin* FUAFLayerAssetProvider::CreateLayerContentTrait(UE::UAF::Layering::FL
 			LayerCreationContext.CompileSettings.ReportError(TEXT("Failed to create Layer Asset DataTrait."));
 			return nullptr;
 		}
+		
+		URigVMPin* TraitPin = ContentHostTraitStackNode->FindTrait(AssetDataTraitName);
+		LayerCreationContext.GraphController->SetPinExpansion(TraitPin->GetPinPath(), true, false);
 	}
 
 	URigVMPin* BlendTraitStackResult = ContentHostTraitStackNode->FindPin(GET_MEMBER_NAME_STRING_CHECKED(FRigUnit_AnimNextTraitStack, Result));
@@ -83,18 +86,24 @@ URigVMPin* FUAFLayerAssetProvider::CreateLayerContentTrait(UE::UAF::Layering::FL
 
 void FUAFLayerAssetProvider::SetLayerAsset(UObject* InNewAsset)
 {
+	const FScopedTransaction Transaction(LOCTEXT("OnLayerAssetSelectionChangedAction", "Set Layer Asset"));
+	if (OuterLayer.IsValid())
+	{
+		OuterLayer->Modify();
+	}
+	
 	LayerAsset = InNewAsset;
+	
+	if (OuterLayer.IsValid())
+	{
+		OuterLayer->BroadcastModified(EAnimNextEditorDataNotifType::PropertyChanged);
+	}
 }
 
 void FUAFLayerAssetProvider::OnLayerAssetSelectionChanged(const FAssetData& InAssetData)
 {
 	UObject* Obj = Cast<UObject>(InAssetData.GetAsset());
 	SetLayerAsset(Obj);
-	
-	if (OuterLayer.IsValid())
-	{
-		OuterLayer->BroadcastModified(EAnimNextEditorDataNotifType::PropertyChanged);
-	}
 }
 
 TSharedRef<SWidget> FUAFLayerAssetProvider::CreateLayerContentWidget(UUAFLayer* InLayer)
