UAF Layering
- Added drag and drop operations on layers
- Added additional tool tips
- First pass on undo/redo support
- Added a delta step to adjusting the layer weight, blend in time and blend out time

#Synced-CL 50702987

diff --git a/Engine/Plugins/Experimental/UAF/UAFLayering/Source/UAFLayeringUncookedOnly/Internal/Layers/UAFLayer.h b/Engine/Plugins/Experimental/UAF/UAFLayering/Source/UAFLayeringUncookedOnly/Internal/Layers/UAFLayer.h
--- a/Engine/Plugins/Experimental/UAF/UAFLayering/Source/UAFLayeringUncookedOnly/Internal/Layers/UAFLayer.h
+++ b/Engine/Plugins/Experimental/UAF/UAFLayering/Source/UAFLayeringUncookedOnly/Internal/Layers/UAFLayer.h
@@ -55,7 +55,7 @@ public:
 	TInstancedStruct<FUAFLayerContentProviderBase>& GetLayerContentProvider();
 	void SetLayerContentProvider(const TInstancedStruct<FUAFLayerContentProviderBase>& InLayerContentProvider);
 	
-	void BroadcastModified(EAnimNextEditorDataNotifType InType, bool bNotifyLayerLayoutChanges = false);
+	void BroadcastModified(EAnimNextEditorDataNotifType InType, bool bNotifyLayerLayoutChanges = false) const;
 
 #if WITH_EDITOR
 	virtual void PreEditChange(FProperty* PropertyAboutToChange) override;
@@ -63,8 +63,8 @@ public:
 #endif
 	
 protected:
-	UPROPERTY(EditAnywhere, Category = "Layer")
 	// The current state of the layer. 
+	UPROPERTY(EditAnywhere, Category = "Layer")
 	EUAFLayerState LayerState = EUAFLayerState::Enabled;
 
 	// The name of this layer. It has to be unique within one layer stack asset. 
@@ -77,6 +77,8 @@ protected:
 	UPROPERTY(EditAnywhere, Category = "Layer", meta = (ExcludeBaseStruct))
 	TInstancedStruct<FUAFLayerBlendProviderBase> LayerBlendProvider;
 	
+	FName CachedLayerOutputPoseName = NAME_None;
+	
 	FName PreRenameLayerName = NAME_None;
 };
 
diff --git a/Engine/Plugins/Experimental/UAF/UAFLayering/Source/UAFLayeringUncookedOnly/Internal/Layers/UAFLayerDefaultBlendProvider.h b/Engine/Plugins/Experimental/UAF/UAFLayering/Source/UAFLayeringUncookedOnly/Internal/Layers/UAFLayerDefaultBlendProvider.h
--- a/Engine/Plugins/Experimental/UAF/UAFLayering/Source/UAFLayeringUncookedOnly/Internal/Layers/UAFLayerDefaultBlendProvider.h
+++ b/Engine/Plugins/Experimental/UAF/UAFLayering/Source/UAFLayeringUncookedOnly/Internal/Layers/UAFLayerDefaultBlendProvider.h
@@ -48,21 +48,21 @@ public:
 	
 	// TODO: This should be a binding
 	// The weight to use when blending this layer in. 0 == this layer won't get blended in at all, 1 == 100% of the layer gets blended in 
-	UPROPERTY(EditAnywhere, Category = "Layer", meta=(ClampMin = "0.0", ClampMax = "1.0", UIMin = "0.0", UIMax = "1.0"))
+	UPROPERTY(EditAnywhere, Category = "Layer", meta=(ClampMin = "0.0", ClampMax = "1.0", UIMin = "0.0", UIMax = "1.0", Delta=0.1))
 	float LayerWeight = 1.0f;
 	
 	// TODO: This should be a binding
 	// How long it should take to blend this layer in if it has been disabled previously 
-	UPROPERTY(EditAnywhere, Category = "Layer|BlendSettings")
+	UPROPERTY(EditAnywhere, Category = "Layer|BlendSettings", meta=(Delta=0.1))
 	float LayerBlendInTime = 0.5f;
 
 	// TODO: This should be a binding
 	// How long it should take to blend this layer out 
-	UPROPERTY(EditAnywhere, Category = "Layer|BlendSettings")
+	UPROPERTY(EditAnywhere, Category = "Layer|BlendSettings", meta=(Delta=0.1))
 	float LayerBlendOutTime = 0.5f;
 	
 	// TODO: This should be a binding
-	// A optional blend mask to apply when blending this layers pose result with the previous layer result 
+	// An optional blend mask to apply when blending this layers pose result with the previous layer result 
 	UPROPERTY(EditAnywhere, Category = "Layer")
 	TObjectPtr<UHierarchyTable> BlendMask = nullptr;
 	
diff --git a/Engine/Plugins/Experimental/UAF/UAFLayering/Source/UAFLayeringUncookedOnly/Internal/UAFLayerStack_EditorData.h b/Engine/Plugins/Experimental/UAF/UAFLayering/Source/UAFLayeringUncookedOnly/Internal/UAFLayerStack_EditorData.h
--- a/Engine/Plugins/Experimental/UAF/UAFLayering/Source/UAFLayeringUncookedOnly/Internal/UAFLayerStack_EditorData.h
+++ b/Engine/Plugins/Experimental/UAF/UAFLayering/Source/UAFLayeringUncookedOnly/Internal/UAFLayerStack_EditorData.h
@@ -2,6 +2,7 @@
 
 #pragma once
 
+#include "EditorUndoClient.h"
 #include "LayeringUncookedOnlyTypes.h"
 #include "Graph/AnimNextAnimationGraph_EditorData.h"
 #include "UAFLayerStack_EditorData.generated.h"
@@ -15,25 +16,35 @@ class UUAFLayerStack;
 DECLARE_MULTICAST_DELEGATE(FLayerLayoutChangedDelegate);
 
 UCLASS()
-class UUAFLayerStack_EditorData : public UUAFAnimGraph_EditorData
+class UUAFLayerStack_EditorData : public UUAFAnimGraph_EditorData, public FEditorUndoClient
 {
 	GENERATED_BODY()
 
 public:
+	UE_API UUAFLayerStack_EditorData();
+	UE_API virtual ~UUAFLayerStack_EditorData() override;
+	
+	virtual void PostUndo(bool bSuccess) override;
+	virtual void PostRedo(bool bSuccess) override;
+	
 	int32 GetNumLayers() const;
 	TArray<TObjectPtr<UUAFLayer>> GetAllLayers() const;
 	int32 GetIndexForLayer(UUAFLayer* Layer) const;
+	int32 GetIndexForMoveableLayer(UUAFLayer* Layer) const;
 	bool IsBaseLayer(UUAFLayer* Layer) const;
 	bool IsLastLayer(UUAFLayer* Layer) const;
 	bool IsLayerNameValid(const FName NewLayerName) const;
 	
+	void MoveLayerToIndex(UUAFLayer* LayerToMove, int32 NewLayerIndex);
 	void MoveLayerUp(UUAFLayer* LayerToMove);
 	void MoveLayerDown(UUAFLayer* LayerToMove);
 	void DeleteLayer(UUAFLayer* LayerToDelete);
 	void SetLayerState(UUAFLayer* Layer, EUAFLayerState LayerState);
 	
 	UUAFLayer* CreateDefaultAssetBasedLayer(const FAssetData& InAsset);
-	void NotifyLayerLayoutChanged();
+	
+	void BroadcastLayerStackChanged(EAnimNextEditorDataNotifType InType, bool bRefreshUI = false);
+	void NotifyLayerLayoutChanged() const;
 	
 protected:
 	// Begin UUAFAnimGraph_EditorData interface
diff --git a/Engine/Plugins/Experimental/UAF/UAFLayering/Source/UAFLayeringUncookedOnly/Internal/Widgets/SDefaultBlendProvider.h b/Engine/Plugins/Experimental/UAF/UAFLayering/Source/UAFLayeringUncookedOnly/Internal/Widgets/SDefaultBlendProvider.h
--- a/Engine/Plugins/Experimental/UAF/UAFLayering/Source/UAFLayeringUncookedOnly/Internal/Widgets/SDefaultBlendProvider.h
+++ b/Engine/Plugins/Experimental/UAF/UAFLayering/Source/UAFLayeringUncookedOnly/Internal/Widgets/SDefaultBlendProvider.h
@@ -4,6 +4,7 @@
 
 #include "CoreMinimal.h"
 #include "ISinglePropertyView.h"
+#include "Misc/NotifyHook.h"
 #include "Widgets/SCompoundWidget.h"
 
 namespace UE::UAF::Layering
@@ -16,7 +17,7 @@ class UUAFLayer;
 
 namespace UE::UAF::Layering
 {
-class SDefaultBlendProvider : public SCompoundWidget
+class SDefaultBlendProvider : public SCompoundWidget, public FNotifyHook
 {
 public:
 	SLATE_BEGIN_ARGS(SDefaultBlendProvider) {}
@@ -26,9 +27,9 @@ public:
 	/** Constructs this widget with InArgs */
 	void Construct(const FArguments& InArgs);
 	
-private:
-	TSharedPtr<ISinglePropertyView> GenerateSinglePropertyView(const FName PropertyName, const TSharedRef<FDefaultLayerBlendStructureDataProvider> StructDataProvider, const FSinglePropertyParams& Params) const;
-	
+	virtual void NotifyPreChange( class FEditPropertyChain* PropertyAboutToChange ) override;
+	virtual void NotifyPostChange( const FPropertyChangedEvent& PropertyChangedEvent, class FEditPropertyChain* PropertyThatChanged ) override;
+
 private:
 	TWeakObjectPtr<UUAFLayer> WeakLayer = nullptr;
 };
diff --git a/Engine/Plugins/Experimental/UAF/UAFLayering/Source/UAFLayeringUncookedOnly/Internal/Widgets/SLayer.h b/Engine/Plugins/Experimental/UAF/UAFLayering/Source/UAFLayeringUncookedOnly/Internal/Widgets/SLayer.h
--- a/Engine/Plugins/Experimental/UAF/UAFLayering/Source/UAFLayeringUncookedOnly/Internal/Widgets/SLayer.h
+++ b/Engine/Plugins/Experimental/UAF/UAFLayering/Source/UAFLayeringUncookedOnly/Internal/Widgets/SLayer.h
@@ -25,7 +25,6 @@ private:
 	const FSlateBrush* GetBackgroundOverlayBrush() const;
 	
 	FSlateColor GetIndicatorColor() const;
-	FText GetLayerNameTooltip() const;
 	FText GetLayerName() const;
 	bool IsSelected() const;
 	bool VerifyLayerName(const FText& Text, FText& OutErrorMessage) const;
diff --git a/Engine/Plugins/Experimental/UAF/UAFLayering/Source/UAFLayeringUncookedOnly/Private/Layers/UAFLayer.cpp b/Engine/Plugins/Experimental/UAF/UAFLayering/Source/UAFLayeringUncookedOnly/Private/Layers/UAFLayer.cpp
--- a/Engine/Plugins/Experimental/UAF/UAFLayering/Source/UAFLayeringUncookedOnly/Private/Layers/UAFLayer.cpp
+++ b/Engine/Plugins/Experimental/UAF/UAFLayering/Source/UAFLayeringUncookedOnly/Private/Layers/UAFLayer.cpp
@@ -63,15 +63,11 @@ void UUAFLayer::TogglePreviewVisibility()
 	}
 }
 
-void UUAFLayer::BroadcastModified(EAnimNextEditorDataNotifType InType, bool bNotifyLayerLayoutChanges)
+void UUAFLayer::BroadcastModified(EAnimNextEditorDataNotifType InType, bool bNotifyLayerLayoutChanges) const
 {
 	if (UUAFLayerStack_EditorData* OuterEditorData = Cast<UUAFLayerStack_EditorData>(GetOuter()))
 	{
-		OuterEditorData->BroadcastModified(InType, this);
-		if (bNotifyLayerLayoutChanges)
-		{
-			OuterEditorData->NotifyLayerLayoutChanged();
-		}
+		OuterEditorData->BroadcastLayerStackChanged(InType, bNotifyLayerLayoutChanges);
 	}
 }
 
diff --git a/Engine/Plugins/Experimental/UAF/UAFLayering/Source/UAFLayeringUncookedOnly/Private/UAFLayerStack_EditorData.cpp b/Engine/Plugins/Experimental/UAF/UAFLayering/Source/UAFLayeringUncookedOnly/Private/UAFLayerStack_EditorData.cpp
--- a/Engine/Plugins/Experimental/UAF/UAFLayering/Source/UAFLayeringUncookedOnly/Private/UAFLayerStack_EditorData.cpp
+++ b/Engine/Plugins/Experimental/UAF/UAFLayering/Source/UAFLayeringUncookedOnly/Private/UAFLayerStack_EditorData.cpp
@@ -2,6 +2,8 @@
 
 #include "UAFLayerStack_EditorData.h"
 #include "AnimGraphUncookedOnlyUtils.h"
+#include "Editor.h"
+#include "ScopedTransaction.h"
 #include "UAFLayerStack.h"
 #include "Compilation/AnimNextGetGraphCompileContext.h"
 #include "Entries/AnimNextAnimationGraphEntry.h"
@@ -15,6 +17,37 @@
 #include "Layers/UAFLayerDefaultBlendProvider.h"
 #include "Workspace/UAFLayerStackWorkspaceAssetUserData.h"
 
+#define LOCTEXT_NAMESPACE "UUAFLayerStack_EditorData"
+
+UUAFLayerStack_EditorData::UUAFLayerStack_EditorData()
+{
+	if (GEditor)
+	{
+		GEditor->RegisterForUndo(this);
+	}
+}
+
+UUAFLayerStack_EditorData::~UUAFLayerStack_EditorData()
+{
+	if (GEditor)
+	{
+		GEditor->UnregisterForUndo(this);
+	}
+}
+
+void UUAFLayerStack_EditorData::PostUndo(bool bSuccess)
+{
+	if (bSuccess)
+	{
+		BroadcastLayerStackChanged(EAnimNextEditorDataNotifType::UndoRedo, true);
+	}
+}
+
+void UUAFLayerStack_EditorData::PostRedo(bool bSuccess)
+{
+	PostUndo(bSuccess);
+}
+
 void UUAFLayerStack_EditorData::OnPreCompileGetProgrammaticGraphs(const FRigVMCompileSettings& InSettings, FAnimNextGetGraphCompileContext& OutCompileContext)
 {
 	Super::OnPreCompileGetProgrammaticGraphs(InSettings, OutCompileContext);
@@ -91,7 +124,7 @@ void UUAFLayerStack_EditorData::PostEditChangeProperty(FPropertyChangedEvent& Pr
 	}
 }
 
-void UUAFLayerStack_EditorData::NotifyLayerLayoutChanged()
+void UUAFLayerStack_EditorData::NotifyLayerLayoutChanged() const
 {
 	OnLayerLayoutChanged.Broadcast();
 }
@@ -126,11 +159,13 @@ void UUAFLayerStack_EditorData::MoveLayerUp(UUAFLayer* LayerToMove)
 		return;
 	}
 
+	const FScopedTransaction Transaction(LOCTEXT("MoveLayerUpAction", "Move Layer Up"));
+	Modify();
+	
 	const int32 TargetIndex = SourceLayerIndex - 1;
 	Layers.Swap(SourceLayerIndex, TargetIndex);
 
-	NotifyLayerLayoutChanged();
-	BroadcastModified(EAnimNextEditorDataNotifType::PropertyChanged, this);
+	BroadcastLayerStackChanged(EAnimNextEditorDataNotifType::PropertyChanged, true);
 }
 
 void UUAFLayerStack_EditorData::MoveLayerDown(UUAFLayer* LayerToMove)
@@ -154,11 +189,13 @@ void UUAFLayerStack_EditorData::MoveLayerDown(UUAFLayer* LayerToMove)
 		return;
 	}
 
+	const FScopedTransaction Transaction(LOCTEXT("MoveLayerDownAction", "Move Layer Down"));
+	Modify();
+	
 	const int32 TargetIndex = SourceLayerIndex + 1;
 	Layers.Swap(SourceLayerIndex, TargetIndex);
 
-	NotifyLayerLayoutChanged();
-	BroadcastModified(EAnimNextEditorDataNotifType::PropertyChanged, this);
+	BroadcastLayerStackChanged(EAnimNextEditorDataNotifType::PropertyChanged, true);
 }
 
 void UUAFLayerStack_EditorData::DeleteLayer(UUAFLayer* LayerToDelete)
@@ -170,15 +207,21 @@ void UUAFLayerStack_EditorData::DeleteLayer(UUAFLayer* LayerToDelete)
 			// not allowed to remove the base layer
 			return;
 		}
+		
+		const FScopedTransaction Transaction(LOCTEXT("DeleteLayerAction", "Delete Layer"));
+		Modify();
+		
 		Layers.Remove(LayerToDelete);
 		
-		NotifyLayerLayoutChanged();
-		BroadcastModified(EAnimNextEditorDataNotifType::PropertyChanged, this);
+		BroadcastLayerStackChanged(EAnimNextEditorDataNotifType::PropertyChanged, true);
 	}
 }
 
 void UUAFLayerStack_EditorData::SetLayerState(UUAFLayer* Layer, EUAFLayerState LayerState)
 {
+	const FScopedTransaction Transaction(LOCTEXT("SetLayerStateAction", "Set Layer State"));
+	Modify();
+	
 	Layer->SetLayerState(LayerState);
 	BroadcastModified(EAnimNextEditorDataNotifType::PropertyChanged, this);
 }
@@ -188,6 +231,11 @@ int32 UUAFLayerStack_EditorData::GetIndexForLayer(UUAFLayer* Layer) const
 	return GetAllLayers().Find(Layer);
 }
 
+int32 UUAFLayerStack_EditorData::GetIndexForMoveableLayer(UUAFLayer* Layer) const 
+{
+	return Layers.Find(Layer);
+}
+
 int32 UUAFLayerStack_EditorData::GetNumLayers() const
 {
 	return GetAllLayers().Num();
@@ -229,8 +277,50 @@ bool UUAFLayerStack_EditorData::IsLayerNameValid(const FName NewLayerName) const
 	return true;
 }
 
+void UUAFLayerStack_EditorData::MoveLayerToIndex(UUAFLayer* LayerToMove, int32 NewLayerIndex)
+{
+	if (LayerToMove && !Layers.IsEmpty())
+	{
+		NewLayerIndex = FMath::Clamp(NewLayerIndex, 0, Layers.Num() - 1);
+		int32 OriginalLayerIndex = GetIndexForMoveableLayer(LayerToMove);
+		if (OriginalLayerIndex == NewLayerIndex || OriginalLayerIndex == INDEX_NONE)
+		{
+			return;
+		}
+		
+		const FScopedTransaction Transaction(LOCTEXT("MoveLayerToIndexAction", "Move Layer to Index"));
+		Modify();
+
+		if (NewLayerIndex > OriginalLayerIndex)
+		{
+			NewLayerIndex += 1;
+		}
+		else
+		{
+			OriginalLayerIndex += 1;
+		}
+		
+		if (NewLayerIndex > Layers.Num())
+		{
+			Layers.AddZeroed();
+		}
+		else
+		{
+			Layers.InsertZeroed(NewLayerIndex);
+		}
+		
+		Layers.Swap(NewLayerIndex, OriginalLayerIndex);
+		Layers.RemoveAt(OriginalLayerIndex);
+		
+		BroadcastLayerStackChanged(EAnimNextEditorDataNotifType::PropertyChanged, true);
+	}
+}
+
 UUAFLayer* UUAFLayerStack_EditorData::CreateDefaultAssetBasedLayer(const FAssetData& InAsset)
 {
+	const FScopedTransaction Transaction(LOCTEXT("AddLayerAction", "Add Layer"));
+	Modify();
+	
 	UUAFLayer* NewLayer = NewObject<UUAFLayer>(this);
 	if (!NewLayer)
 	{
@@ -251,12 +341,20 @@ UUAFLayer* UUAFLayerStack_EditorData::CreateDefaultAssetBasedLayer(const FAssetD
 	NewLayer->SetLayerContentProvider(TInstancedStruct<FUAFLayerAssetProvider>::Make(NewContentProvider));
 	NewLayer->SetLayerBlendProvider(TInstancedStruct<FUAFDefaultBlendProvider>::Make());
 	
-	NotifyLayerLayoutChanged();
-	BroadcastModified(EAnimNextEditorDataNotifType::PropertyChanged, this);
+	BroadcastLayerStackChanged(EAnimNextEditorDataNotifType::PropertyChanged, true);
 	
 	return NewLayer;
 }
 
+void UUAFLayerStack_EditorData::BroadcastLayerStackChanged(EAnimNextEditorDataNotifType InType, bool bRefreshUI)
+{
+	BroadcastModified(InType, this);
+	if (bRefreshUI)
+	{
+		NotifyLayerLayoutChanged();
+	}
+}
+
 void UUAFLayerStack_EditorData::Initialize(bool bRecompileVM)
 {
 	Super::Initialize(bRecompileVM);
@@ -283,3 +381,5 @@ TConstArrayView<TSubclassOf<UUAFRigVMAssetEntry>> UUAFLayerStack_EditorData::Get
 
 	return Classes;
 }
+
+#undef LOCTEXT_NAMESPACE
diff --git a/Engine/Plugins/Experimental/UAF/UAFLayering/Source/UAFLayeringUncookedOnly/Private/UAFLayeringStyle.cpp b/Engine/Plugins/Experimental/UAF/UAFLayering/Source/UAFLayeringUncookedOnly/Private/UAFLayeringStyle.cpp
--- a/Engine/Plugins/Experimental/UAF/UAFLayering/Source/UAFLayeringUncookedOnly/Private/UAFLayeringStyle.cpp
+++ b/Engine/Plugins/Experimental/UAF/UAFLayering/Source/UAFLayeringUncookedOnly/Private/UAFLayeringStyle.cpp
@@ -3,6 +3,7 @@
 #include "UAFLayeringStyle.h"
 #include "Styling/SlateStyleRegistry.h"
 #include "Brushes/SlateRoundedBoxBrush.h"
+#include "Styling/AppStyle.h"
 #include "Styling/StyleColors.h"
 #include "Styling/SlateStyleMacros.h"
 #include "Styling/SlateTypes.h"
@@ -31,6 +32,7 @@ FUAFLayeringStyle::FUAFLayeringStyle()
 	Set("Layer.WhiteBackground", new FSlateRoundedBoxBrush(FLinearColor::White, 5.0f));
 	Set("Layer.AdvancedSettingsBackground", new FSlateRoundedBoxBrush(FStyleColors::Recessed.GetSpecifiedColor(), FVector4(0.0f, 0.0f, 10.0f, 10.0f))); //. X = Top Left, Y = Top Right, Z = Bottom Right, W = Bottom Left
 	Set("Layer.SelectedLayer", new FSlateRoundedBoxBrush(FLinearColor(0.5f, 0.3f, 0.8f), 5.0f));
+	Set("Layer.DraggedLayerBackground", new FSlateBoxBrush(NAME_None,FMargin(10.0f, 10.0f), FStyleColors::Recessed.GetSpecifiedColor()));
 	
 	Set("LayerStack.MiddleBackground", new FSlateRoundedBoxBrush(FStyleColors::AccentGray.GetSpecifiedColor(), 10.0f));
 	
@@ -44,6 +46,9 @@ FUAFLayeringStyle::FUAFLayeringStyle()
 	.SetNormalPadding(FMargin(0, 0, 0, 1))
 	.SetPressedPadding(FMargin(0, 1, 0, 0)));
 	
+	const FVector2D Icon16x16(16.0f, 16.0f);
+	const FVector2D Icon64x64(64.0f, 64.0f);
+	const FVector2D Icon32x32(32.0f, 32.0f);
 	
 	Set("LayerStack.ListViewRow", FTableRowStyle()
 		.SetEvenRowBackgroundBrush(FSlateNoResource())
@@ -51,17 +56,16 @@ FUAFLayeringStyle::FUAFLayeringStyle()
 		.SetOddRowBackgroundBrush(FSlateNoResource())
 		.SetOddRowBackgroundHoveredBrush(FSlateNoResource())
 		.SetActiveHighlightedBrush(FSlateNoResource())
-		.SetActiveBrush(FSlateNoResource())
-		.SetActiveHoveredBrush(FSlateNoResource())
-		.SetSelectorFocusedBrush(FSlateRoundedBoxBrush(FStyleColors::AccentBlue.GetSpecifiedColor(), 10.0f))
+		.SetActiveBrush(FSlateRoundedBoxBrush(FStyleColors::AccentBlue.GetSpecifiedColor(), 10.0f))
+		.SetActiveHoveredBrush(FSlateRoundedBoxBrush(FStyleColors::AccentBlue.GetSpecifiedColor(), 10.0f))
+		.SetSelectorFocusedBrush(FSlateNoResource())
 		.SetInactiveHighlightedBrush(FSlateNoResource())
 		.SetInactiveBrush(FSlateRoundedBoxBrush(FStyleColors::AccentBlue.GetSpecifiedColor(), 10.0f))
 		.SetInactiveHoveredBrush(FSlateRoundedBoxBrush(FStyleColors::AccentBlue.GetSpecifiedColor(), 10.0f))
-	);
-	
-	const FVector2D Icon16x16(16.0f, 16.0f);
-	const FVector2D Icon64x64(64.0f, 64.0f);
-	const FVector2D Icon32x32(32.0f, 32.0f);
+		.SetDropIndicator_Above(CORE_BOX_BRUSH("Common/DropZoneIndicator_Above", FMargin(10.0f / 16.0f, 10.0f / 16.0f, 0, 0), FStyleColors::AccentOrange.GetSpecifiedColor()))
+		.SetDropIndicator_Onto(CORE_BOX_BRUSH("Common/DropZoneIndicator_Onto", FMargin(4.0f / 16.0f), FStyleColors::AccentOrange.GetSpecifiedColor()))
+		.SetDropIndicator_Below(CORE_BOX_BRUSH("Common/DropZoneIndicator_Below", FMargin(10.0f / 16.0f, 0, 0, 10.0f / 16.0f), FStyleColors::AccentOrange.GetSpecifiedColor()))
+		);
 	
 	Set("ClassIcon.UAFLayerStack", new IMAGE_BRUSH("Icons/layers_16x", Icon16x16));
 	
diff --git a/Engine/Plugins/Experimental/UAF/UAFLayering/Source/UAFLayeringUncookedOnly/Private/Widgets/LayerDragDropOp.h b/Engine/Plugins/Experimental/UAF/UAFLayering/Source/UAFLayeringUncookedOnly/Private/Widgets/LayerDragDropOp.h
new file mode 100644
--- /dev/null
+++ b/Engine/Plugins/Experimental/UAF/UAFLayering/Source/UAFLayeringUncookedOnly/Private/Widgets/LayerDragDropOp.h
@@ -0,0 +1,67 @@
+// Copyright Epic Games, Inc. All Rights Reserved.
+
+#pragma once
+
+#include "CoreMinimal.h"
+#include "UAFLayeringStyle.h"
+#include "DragAndDrop/DecoratedDragDropOp.h"
+#include "Layers/UAFLayer.h"
+#include "Widgets/SOverlay.h"
+
+class UUAFLayer;
+
+#define LOCTEXT_NAMESPACE "FLayerDragDropOp"
+
+namespace UE::UAF::Layering
+{
+	class FLayerDragDropOp : public FDecoratedDragDropOp
+	{
+	public:
+		DRAG_DROP_OPERATOR_TYPE(FLayerDragDropOp, FDecoratedDragDropOp)
+		
+		static TSharedRef<FLayerDragDropOp> New(const TWeakObjectPtr<UUAFLayer> InLayer)
+		{
+			TSharedRef<FLayerDragDropOp> Operation = MakeShared<FLayerDragDropOp>();
+			Operation->DraggedLayer = InLayer;
+
+			Operation->Construct();
+
+			return Operation;
+		}
+		
+		virtual TSharedPtr<SWidget> GetDefaultDecorator() const override
+		{
+			return SNew(SOverlay)
+				+ SOverlay::Slot()
+				.HAlign(HAlign_Fill)
+				.VAlign(VAlign_Fill)
+				[
+					SNew(SBorder)
+					.BorderImage(FUAFLayeringStyle::Get().GetBrush(TEXT("Layer.DraggedLayerBackground")))
+					.Padding(10.0f)
+						[
+					
+						SNew(SBorder)
+						.BorderImage(FUAFLayeringStyle::Get().GetBrush(TEXT("Layer.Background")))
+						[
+							SNew(SVerticalBox)
+							+ SVerticalBox::Slot()
+							.Padding(20.0f, 10.0f)
+							[
+								SNew(STextBlock)
+								.Text_Lambda([this]()
+									{
+										return DraggedLayer.IsValid() ? FText::FromName(DraggedLayer->GetLayerName()) : LOCTEXT("InvalidDraggedLayerLabel", "Invalid Layer");
+									})
+							]
+						]
+					]
+				];	
+		}
+		
+		TWeakObjectPtr<UUAFLayer> DraggedLayer = nullptr;
+	};
+}
+
+#undef LOCTEXT_NAMESPACE
+
diff --git a/Engine/Plugins/Experimental/UAF/UAFLayering/Source/UAFLayeringUncookedOnly/Private/Widgets/SDefaultBlendProvider.cpp b/Engine/Plugins/Experimental/UAF/UAFLayering/Source/UAFLayeringUncookedOnly/Private/Widgets/SDefaultBlendProvider.cpp
--- a/Engine/Plugins/Experimental/UAF/UAFLayering/Source/UAFLayeringUncookedOnly/Private/Widgets/SDefaultBlendProvider.cpp
+++ b/Engine/Plugins/Experimental/UAF/UAFLayering/Source/UAFLayeringUncookedOnly/Private/Widgets/SDefaultBlendProvider.cpp
@@ -2,7 +2,9 @@
 
 #include "Widgets/SDefaultBlendProvider.h"
 
+#include "Editor.h"
 #include "ISinglePropertyView.h"
+#include "ScopedTransaction.h"
 #include "SlateOptMacros.h"
 #include "Layers/UAFLayer.h"
 #include "Layers/UAFLayerDefaultBlendProvider.h"
@@ -28,21 +30,36 @@ void SDefaultBlendProvider::Construct(const FArguments& InArgs)
 	DefaultPropertyArgs.NamePlacement = EPropertyNamePlacement::Hidden;
 	DefaultPropertyArgs.bHideResetToDefault = true;
 	DefaultPropertyArgs.bHideAssetThumbnail = true;
+	DefaultPropertyArgs.NotifyHook = this;
 	
+	FPropertyEditorModule& PropertyEditorModule = FModuleManager::Get().LoadModuleChecked<FPropertyEditorModule>("PropertyEditor");
 	TSharedRef<FDefaultLayerBlendStructureDataProvider> LayerBlendDataProvider = MakeShared<FDefaultLayerBlendStructureDataProvider>(&WeakLayer->GetLayerBlendProvider());
 	
 	// Generate single property views to display blend properties within the layer widget 
-	TSharedPtr<ISinglePropertyView> LayerBlendModePropertyView = GenerateSinglePropertyView(GET_MEMBER_NAME_CHECKED(FUAFDefaultBlendProvider, BlendMode), LayerBlendDataProvider, DefaultPropertyArgs);
-	TSharedPtr<ISinglePropertyView> LayerWeightPropertyView = GenerateSinglePropertyView(GET_MEMBER_NAME_CHECKED(FUAFDefaultBlendProvider, LayerWeight), LayerBlendDataProvider, DefaultPropertyArgs);
-
-	TSharedPtr<ISinglePropertyView> LayerMaskPropertyView = GenerateSinglePropertyView(GET_MEMBER_NAME_CHECKED(FUAFDefaultBlendProvider, BlendMask), LayerBlendDataProvider, DefaultPropertyArgs);
-	TSharedPtr<ISinglePropertyView> BlendCurvePropertyView = GenerateSinglePropertyView(GET_MEMBER_NAME_CHECKED(FUAFDefaultBlendProvider, BlendCurve), LayerBlendDataProvider, DefaultPropertyArgs);
-
+	TSharedPtr<ISinglePropertyView> LayerBlendModePropertyView = PropertyEditorModule.CreateSingleProperty(LayerBlendDataProvider, GET_MEMBER_NAME_CHECKED(FUAFDefaultBlendProvider, BlendMode), DefaultPropertyArgs);
+	TSharedPtr<ISinglePropertyView> LayerWeightPropertyView =  PropertyEditorModule.CreateSingleProperty(LayerBlendDataProvider, GET_MEMBER_NAME_CHECKED(FUAFDefaultBlendProvider, LayerWeight), DefaultPropertyArgs);
+	if (LayerWeightPropertyView)
+	{
+		LayerWeightPropertyView->SetToolTipText(LOCTEXT("LayerWeightTooltip", "The weight to use to apply this layer. 0 == this layer won't get applied at all, 1 == 100% of the layer gets applied"));
+	}
+	
+	TSharedPtr<ISinglePropertyView> LayerMaskPropertyView =  PropertyEditorModule.CreateSingleProperty(LayerBlendDataProvider, GET_MEMBER_NAME_CHECKED(FUAFDefaultBlendProvider, BlendMask), DefaultPropertyArgs);
+	if (LayerMaskPropertyView)
+	{
+		LayerMaskPropertyView->SetToolTipText(LOCTEXT("BlendMaskToolTip", "An optional blend mask to apply when blending this layers pose result with the previous layers pose"));
+	}
+	
+	TSharedPtr<ISinglePropertyView> BlendCurvePropertyView =  PropertyEditorModule.CreateSingleProperty(LayerBlendDataProvider, GET_MEMBER_NAME_CHECKED(FUAFDefaultBlendProvider, BlendCurve), DefaultPropertyArgs);
+	if (BlendCurvePropertyView)
+	{
+		BlendCurvePropertyView->SetToolTipText(LOCTEXT("BlendCurveToolTip", "The custom blend curve to use when blending this layer in or out"));
+	}
+	
 	// Display the name for the following property views
 	DefaultPropertyArgs.NamePlacement = EPropertyNamePlacement::Left;
-	TSharedPtr<ISinglePropertyView> LayerEnabledPropertyView = GenerateSinglePropertyView(GET_MEMBER_NAME_CHECKED(FUAFDefaultBlendProvider, bLayerEnabled), LayerBlendDataProvider, DefaultPropertyArgs);
-	TSharedPtr<ISinglePropertyView> LayerBlendInTimePropertyView = GenerateSinglePropertyView(GET_MEMBER_NAME_CHECKED(FUAFDefaultBlendProvider, LayerBlendInTime), LayerBlendDataProvider, DefaultPropertyArgs);
-	TSharedPtr<ISinglePropertyView> LayerBlendOutTimePropertyView = GenerateSinglePropertyView(GET_MEMBER_NAME_CHECKED(FUAFDefaultBlendProvider, LayerBlendOutTime), LayerBlendDataProvider, DefaultPropertyArgs);
+	TSharedPtr<ISinglePropertyView> LayerEnabledPropertyView =  PropertyEditorModule.CreateSingleProperty(LayerBlendDataProvider, GET_MEMBER_NAME_CHECKED(FUAFDefaultBlendProvider, bLayerEnabled), DefaultPropertyArgs);
+	TSharedPtr<ISinglePropertyView> LayerBlendInTimePropertyView =  PropertyEditorModule.CreateSingleProperty(LayerBlendDataProvider, GET_MEMBER_NAME_CHECKED(FUAFDefaultBlendProvider, LayerBlendInTime), DefaultPropertyArgs);
+	TSharedPtr<ISinglePropertyView> LayerBlendOutTimePropertyView =  PropertyEditorModule.CreateSingleProperty(LayerBlendDataProvider, GET_MEMBER_NAME_CHECKED(FUAFDefaultBlendProvider, LayerBlendOutTime), DefaultPropertyArgs);
 
 	ChildSlot
 	[
@@ -62,16 +79,17 @@ void SDefaultBlendProvider::Construct(const FArguments& InArgs)
 				[
 					SNew(STextBlock)
 					.Text(LOCTEXT("BlendModeLabel", "Blend Mode"))
+					.ToolTipText(LOCTEXT("BlendModeLabelToolTip", "The blend mode this layer uses to blend its result with the pose result of the previous layer"))
 				]
 				+ SVerticalBox::Slot()
 				.Padding(0.0f, 2.0f)
 				[
-					LayerBlendModePropertyView.ToSharedRef()
+					LayerBlendModePropertyView ? LayerBlendModePropertyView.ToSharedRef() : SNullWidget::NullWidget
 				]
 				+ SVerticalBox::Slot()
 				.Padding(0.0f, 2.0f)
 				[
-					LayerMaskPropertyView.ToSharedRef()
+					LayerMaskPropertyView ? LayerMaskPropertyView.ToSharedRef() : SNullWidget::NullWidget
 				]
 			]
 
@@ -88,11 +106,12 @@ void SDefaultBlendProvider::Construct(const FArguments& InArgs)
 				[
 					SNew(STextBlock)
 					.Text(LOCTEXT("WeightLabel", "Weight"))
+					.ToolTipText(LOCTEXT("LayerWeightLabelTooltip", "The weight to use to apply this layer. 0 == this layer won't get applied at all, 1 == 100% of the layer gets applied"))
 				]
 				+ SVerticalBox::Slot()
 				.Padding(0.0f, 2.0f)
 				[
-					LayerWeightPropertyView.ToSharedRef()
+					LayerWeightPropertyView ? LayerWeightPropertyView.ToSharedRef() : SNullWidget::NullWidget
 				]
 				+ SVerticalBox::Slot()
 				.Padding(5.0f, 2.0f)
@@ -116,6 +135,7 @@ void SDefaultBlendProvider::Construct(const FArguments& InArgs)
 				[
 					SNew(STextBlock)
 					.Text(LOCTEXT("BlendSettingLabel", "Blend Settings"))
+					.ToolTipText(LOCTEXT("BlendSettingLabelTooltip", "These settings drive the discrete blend when a layer gets dynamically enabled or disabled during runtime"))
 				]
 				+ SVerticalBox::Slot()
 				.Padding(0.0f, 2.0f)
@@ -123,11 +143,11 @@ void SDefaultBlendProvider::Construct(const FArguments& InArgs)
 					SNew(SHorizontalBox)
 					+ SHorizontalBox::Slot()
 					[
-						LayerEnabledPropertyView.ToSharedRef()
+						LayerEnabledPropertyView ? LayerEnabledPropertyView.ToSharedRef() : SNullWidget::NullWidget
 					]
 					+ SHorizontalBox::Slot()
 					[
-						BlendCurvePropertyView.ToSharedRef()
+						BlendCurvePropertyView ? BlendCurvePropertyView.ToSharedRef() : SNullWidget::NullWidget
 					]
 
 				]
@@ -137,11 +157,11 @@ void SDefaultBlendProvider::Construct(const FArguments& InArgs)
 					SNew(SHorizontalBox)
 					+ SHorizontalBox::Slot()
 					[
-						LayerBlendInTimePropertyView.ToSharedRef()
+						LayerBlendInTimePropertyView ? LayerBlendInTimePropertyView.ToSharedRef() : SNullWidget::NullWidget
 					]
 					+ SHorizontalBox::Slot()
 					[
-						LayerBlendOutTimePropertyView.ToSharedRef()
+						LayerBlendOutTimePropertyView ? LayerBlendOutTimePropertyView.ToSharedRef() : SNullWidget::NullWidget
 					]
 
 				]
@@ -151,21 +171,20 @@ void SDefaultBlendProvider::Construct(const FArguments& InArgs)
 	];
 }
 
-TSharedPtr<ISinglePropertyView> SDefaultBlendProvider::GenerateSinglePropertyView(const FName PropertyName, const TSharedRef<FDefaultLayerBlendStructureDataProvider> StructDataProvider, const FSinglePropertyParams& Params) const
+void SDefaultBlendProvider::NotifyPreChange(class FEditPropertyChain* PropertyAboutToChange)
 {
-	FPropertyEditorModule& PropertyEditorModule = FModuleManager::Get().LoadModuleChecked<FPropertyEditorModule>("PropertyEditor");
-
-	TSharedPtr<ISinglePropertyView> GeneratedPropertyView = PropertyEditorModule.CreateSingleProperty(StructDataProvider.ToSharedPtr(), PropertyName, Params);
-	GeneratedPropertyView->SetOnPropertyValueChanged(FSimpleDelegate::CreateLambda([this]
-		{
-			if (WeakLayer.IsValid())
-			{
-				WeakLayer->BroadcastModified(EAnimNextEditorDataNotifType::PropertyChanged);
-			}
-		}));
+	if (WeakLayer.IsValid())
+	{
+		WeakLayer->Modify();
+	}
+}
 
-	return GeneratedPropertyView;
+void SDefaultBlendProvider::NotifyPostChange(const FPropertyChangedEvent& PropertyChangedEvent, class FEditPropertyChain* PropertyThatChanged)
+{
+	if (WeakLayer.IsValid())
+	{
+		WeakLayer->BroadcastModified(EAnimNextEditorDataNotifType::PropertyChanged);
+	}
 }
 }
-
-#undef LOCTEXT_NAMESPACE //UAFDefaultBlendProvider
\ No newline at end of file
+#undef LOCTEXT_NAMESPACE //UAFDefaultBlendProvider
diff --git a/Engine/Plugins/Experimental/UAF/UAFLayering/Source/UAFLayeringUncookedOnly/Private/Widgets/SLayer.cpp b/Engine/Plugins/Experimental/UAF/UAFLayering/Source/UAFLayeringUncookedOnly/Private/Widgets/SLayer.cpp
--- a/Engine/Plugins/Experimental/UAF/UAFLayering/Source/UAFLayeringUncookedOnly/Private/Widgets/SLayer.cpp
+++ b/Engine/Plugins/Experimental/UAF/UAFLayering/Source/UAFLayeringUncookedOnly/Private/Widgets/SLayer.cpp
@@ -51,21 +51,24 @@ void SLayer::Construct(const FArguments& InArgs)
 							SNew(SBorder)
 							.BorderImage(FUAFLayeringStyle::Get().GetBrush(TEXT("Layer.WhiteBackground")))
 							.BorderBackgroundColor(this, &SLayer::GetIndicatorColor)
+							.ToolTipText(LOCTEXT("ColorIndicatorTooltip", "The color indicator for this layer, signifying what operations this layer performs"))
 						]
 						+ SHorizontalBox::Slot()
 						.AutoWidth()
+						.HAlign(HAlign_Fill)
+						.VAlign(VAlign_Fill)
 						[
 							// Grabber Button for drag and drop
-							SNew(SButton)
-							.ContentPadding(FMargin(5.0f, 5.0f, 5.0f, 5.0f))
-							.ButtonStyle(FAppStyle::Get(), "NoBorder")
-							.ForegroundColor(FSlateColor::UseForeground())
+							SNew(SBorder)
+							.BorderImage( FAppStyle::GetBrush("NoBorder") )
+							.Cursor(EMouseCursor::GrabHand)
 							.HAlign(HAlign_Center)
 							.VAlign(VAlign_Center)
-							.Content()
+							.Padding(8.0f)
 							[
 								SNew(SImage)
 								.Image(FAppStyle::GetBrush("VerticalBoxDragIndicatorShort"))
+								
 							]
 						]
 						+ SHorizontalBox::Slot()
@@ -130,7 +133,7 @@ void SLayer::Construct(const FArguments& InArgs)
 								[
 									SNew(SInlineEditableTextBlock)
 									.OnTextCommitted(this, &SLayer::OnLayerNameCommitted)
-									.ToolTipText(this, &SLayer::GetLayerNameTooltip)
+									.ToolTipText(LOCTEXT("LayerNameTooltip", "The name of this layer. The name has to be unique within the layer stack"))
 									.Text(this, &SLayer::GetLayerName)
 									.IsSelected(this, &SLayer::IsSelected)
 									.OnVerifyTextChanged(this, &SLayer::VerifyLayerName)
@@ -155,9 +158,7 @@ void SLayer::Construct(const FArguments& InArgs)
 										.Image(FUAFLayeringStyle::Get().GetBrush(TEXT("Layer.PoseIconSmall")))
 										.ColorAndOpacity(FLinearColor::Yellow)
 									]
-
 								]
-
 							]
 							+ SSplitter::Slot()
 							.Value(0.8)
@@ -179,6 +180,7 @@ void SLayer::Construct(const FArguments& InArgs)
 							.ForegroundColor(FSlateColor::UseForeground())
 							.HAlign(HAlign_Center)
 							.VAlign(VAlign_Center)
+							.ToolTipText(LOCTEXT("AdvancedLayerSettingsToolTip", "Advanced settings for this layer"))
 							.OnClicked_Lambda([this]() -> FReply
 								{
 									bAdvancedSettingsExpanded = !bAdvancedSettingsExpanded;
@@ -219,6 +221,7 @@ void SLayer::Construct(const FArguments& InArgs)
 								[
 									SNew(STextBlock)
 									.Text(LOCTEXT("OutputPoseLabel", "Output Pose"))
+									.ToolTipText(LOCTEXT("OutputPoseLabelToolTip", "The optionally cached pose output of this layer"))
 								]
 								+ SVerticalBox::Slot()
 								.Padding(20.0f, 0.0f)
@@ -240,6 +243,7 @@ void SLayer::Construct(const FArguments& InArgs)
 										SNew(SEditableTextBox)
 										.Text(FText::FromString(TEXT("Output Pose Binding")))
 										.IsEnabled(false)
+										.ToolTipText(LOCTEXT("OutputPoseBindingToolTip", "The property binding the cached output pose of this layer is written to"))
 									]
 								]
 							]
@@ -264,6 +268,7 @@ void SLayer::Construct(const FArguments& InArgs)
 										[
 											SNew(STextBlock)
 											.Text(LOCTEXT("AdditiveOutputPoseBindingLabel", "Output Additive Pose"))
+											.ToolTipText(LOCTEXT("AdditiveOutputPoseLabelToolTip", "The optionally cached additive pose output of this layer and a base"))
 										]
 										+ SHorizontalBox::Slot()
 										.AutoWidth()
@@ -280,6 +285,7 @@ void SLayer::Construct(const FArguments& InArgs)
 									[
 										SNew(STextBlock)
 										.Text(LOCTEXT("AdditiveBaseLabel", "Base"))
+										.ToolTipText(LOCTEXT("AdditiveBaseLabelToolTip", "The base pose for the additive calculation of this layers optional additive output pose"))
 									]
 								]
 								+ SHorizontalBox::Slot()
@@ -295,6 +301,7 @@ void SLayer::Construct(const FArguments& InArgs)
 										SNew(SEditableTextBox)
 										.Text(FText::FromString(TEXT("Additive Pose Binding")))
 										.IsEnabled(false)
+										.ToolTipText(LOCTEXT("AdditivePoseBindingToolTip", "The property binding the cached additive output pose of this layer is written to"))
 									]
 									+ SVerticalBox::Slot()
 									.Padding(5.0f)
@@ -437,12 +444,7 @@ FSlateColor SLayer::GetIndicatorColor() const
 
 	return FSlateColor(FLinearColor::White);
 }
-
-FText SLayer::GetLayerNameTooltip() const
-{
-	return LOCTEXT("LayerNameTooltip", "The name of this layer. The name has to be unique within the layer stack");
-}
-
+	
 FText SLayer::GetLayerName() const
 {
 	if (Layer.IsValid())
diff --git a/Engine/Plugins/Experimental/UAF/UAFLayering/Source/UAFLayeringUncookedOnly/Private/Widgets/SLayerStack.cpp b/Engine/Plugins/Experimental/UAF/UAFLayering/Source/UAFLayeringUncookedOnly/Private/Widgets/SLayerStack.cpp
--- a/Engine/Plugins/Experimental/UAF/UAFLayering/Source/UAFLayeringUncookedOnly/Private/Widgets/SLayerStack.cpp
+++ b/Engine/Plugins/Experimental/UAF/UAFLayering/Source/UAFLayeringUncookedOnly/Private/Widgets/SLayerStack.cpp
@@ -5,6 +5,7 @@
 #include "AnimNextAnimGraphSettings.h"
 #include "UAFLayerStack.h"
 #include "IWorkspaceEditor.h"
+#include "LayerDragDropOp.h"
 #include "UncookedOnlyUtils.h"
 #include "UAFLayerStack_EditorData.h"
 #include "Layers/UAFLayer.h"
@@ -14,6 +15,7 @@
 #include "UAFLayeringStyle.h"
 #include "SAssetDropTarget.h"
 #include "LayeringUncookedOnlyTypes.h"
+#include "ScopedTransaction.h"
 #include "Layers/UAFBaseLayer.h"
 
 #define LOCTEXT_NAMESPACE "FUAFLayerStackWidget"
@@ -107,6 +109,7 @@ void SLayerStack::Construct(const FArguments& InArgs)
 							.VAlign(VAlign_Fill)
 							.HAlign(HAlign_Fill)
 							.BorderImage(FUAFLayeringStyle::Get().GetBrush(TEXT("Layer.AddLayerBackground")))
+							.ToolTipText(LOCTEXT("AddLayerTooltip", "Adds a new layer"))
 							[
 								SNew(SButton)
 								.ContentPadding(FMargin(10.0F))
@@ -229,6 +232,58 @@ TSharedRef<ITableRow> SLayerStack::GenerateLayerRow(TWeakObjectPtr<UUAFLayer> It
 	return SNew(STableRow<TWeakObjectPtr<UUAFLayer>>, OwnerTable)
 		.Padding(3.0f)
 		.Style(FUAFLayeringStyle::Get(), "LayerStack.ListViewRow")
+		.ShowWires(true)
+		.ShowSelection(true)
+		.OnDragDetected_Lambda([Item](const FGeometry& MyGeometry, const FPointerEvent& MouseEvent)
+			{
+				if (MouseEvent.IsMouseButtonDown(EKeys::LeftMouseButton))
+				{
+					// Check if this is a base layer, we won't allow drag and drop operations on these 
+					if (Item.IsValid() && !Item->IsA<UUAFBaseLayer>())
+					{
+						const TSharedRef<FLayerDragDropOp> DragDropOp = FLayerDragDropOp::New(Item);
+						return FReply::Handled().BeginDragDrop(DragDropOp);
+					}
+				}
+
+				return FReply::Unhandled();
+			})
+		.OnCanAcceptDrop_Lambda([](const FDragDropEvent& DragDropEvent, EItemDropZone DropZone, TWeakObjectPtr<UUAFLayer> TargetItem) -> TOptional<EItemDropZone>
+			{
+				TOptional<EItemDropZone> ReturnedDropZone = DropZone;
+				const TSharedPtr<FLayerDragDropOp> LayerDragDropOp = DragDropEvent.GetOperationAs<FLayerDragDropOp>();
+				if (LayerDragDropOp.IsValid())
+				{
+					// If the target is a base layer, we can only drop below it 
+					if (TargetItem.IsValid() && TargetItem->IsA<UUAFBaseLayer>())
+					{
+						ReturnedDropZone = EItemDropZone::BelowItem;
+					}
+				}
+
+				return ReturnedDropZone;
+			})
+		.OnAcceptDrop_Lambda([this](const FDragDropEvent& DragDropEvent, EItemDropZone DropZone, TWeakObjectPtr<UUAFLayer> TargetItem) -> FReply
+			{
+				const TSharedPtr<FLayerDragDropOp> LayerDragDropOp = DragDropEvent.GetOperationAs<FLayerDragDropOp>();
+				if (LayerDragDropOp.IsValid() && LayerStack.IsValid() && LayerDragDropOp->DraggedLayer.IsValid())
+				{
+					UUAFLayerStack_EditorData* EditorData = UAF::UncookedOnly::FUtils::GetEditorData<UUAFLayerStack_EditorData>(LayerStack.Get());
+					if (EditorData)
+					{
+						int32 TargetIndex = EditorData->GetIndexForMoveableLayer(TargetItem.Get());
+						if (DropZone == EItemDropZone::BelowItem)
+						{
+							TargetIndex += 1;
+						}
+					
+						EditorData->MoveLayerToIndex(LayerDragDropOp->DraggedLayer.Get(), TargetIndex);
+						
+						return FReply::Handled();
+					}
+				}
+				return FReply::Unhandled();
+			})
 		[
 			Item.IsValid() ? Item->CreateLayerWidget() : SNullWidget::NullWidget
 		];
