UAF Layering: Expose Asset Data to the Layering UI

#Synced-CL 50858266

diff --git a/Engine/Plugins/Experimental/UAF/UAFLayering/Source/UAFLayering/Internal/Traits/LayerAssetDataTraitData.h b/Engine/Plugins/Experimental/UAF/UAFLayering/Source/UAFLayering/Internal/Traits/LayerAssetDataTraitData.h
--- a/Engine/Plugins/Experimental/UAF/UAFLayering/Source/UAFLayering/Internal/Traits/LayerAssetDataTraitData.h
+++ b/Engine/Plugins/Experimental/UAF/UAFLayering/Source/UAFLayering/Internal/Traits/LayerAssetDataTraitData.h
@@ -1,9 +1,11 @@
 // Copyright Epic Games, Inc. All Rights Reserved.
 
 #pragma once
-#include "Factory/AnimNextFactoryParams.h"
+
+
 #include "TraitCore/TraitBinding.h"
 #include "TraitCore/TraitSharedData.h"
+#include "UAF/UAFAssetData.h"
 
 #include "LayerAssetDataTraitData.generated.h"
 
@@ -14,12 +16,12 @@ struct FUAFLayerAssetDataTraitSharedData : public FAnimNextTraitSharedData
 {
 	GENERATED_BODY()
 	
-	UPROPERTY(EditAnywhere, Category = "Layer|Content", meta = (GetAllowedClasses = "/Script/UAFAnimGraph.AnimNextAnimGraphSettings:GetAllowedAssetClasses"))
-	TObjectPtr<UObject> Asset = nullptr;
+	UPROPERTY(EditAnywhere, Category = "Layer|Content")
+	TInstancedStruct<FUAFGraphFactoryAsset> GraphAssetHandle; 
 	
 	// Latent pin support boilerplate
 	#define TRAIT_LATENT_PROPERTIES_ENUMERATOR(GeneratorMacro) \
-	GeneratorMacro(Asset)
+	GeneratorMacro(GraphAssetHandle)
 
 	GENERATE_TRAIT_LATENT_PROPERTIES(FUAFLayerAssetDataTraitSharedData, TRAIT_LATENT_PROPERTIES_ENUMERATOR)
 	#undef TRAIT_LATENT_PROPERTIES_ENUMERATOR
diff --git a/Engine/Plugins/Experimental/UAF/UAFLayering/Source/UAFLayering/Private/Traits/LayerAssetDataProviderTrait.cpp b/Engine/Plugins/Experimental/UAF/UAFLayering/Source/UAFLayering/Private/Traits/LayerAssetDataProviderTrait.cpp
--- a/Engine/Plugins/Experimental/UAF/UAFLayering/Source/UAFLayering/Private/Traits/LayerAssetDataProviderTrait.cpp
+++ b/Engine/Plugins/Experimental/UAF/UAFLayering/Source/UAFLayering/Private/Traits/LayerAssetDataProviderTrait.cpp
@@ -22,64 +22,72 @@ namespace UE::UAF
 	#undef TRAIT_EVENT_ENUMERATOR
 
 
-void FLayerAssetDataProviderTrait::FInstanceData::Construct(const FExecutionContext& Context, const FTraitBinding& Binding)
-{
-	FTrait::FInstanceData::Construct(Context, Binding);
-}
+	void FLayerAssetDataProviderTrait::FInstanceData::Construct(const FExecutionContext& Context, const FTraitBinding& Binding)
+	{
+		FTrait::FInstanceData::Construct(Context, Binding);
+			
+		CachedAssetHandle = FGraphAssetHandle();
+	}
 
-void FLayerAssetDataProviderTrait::PreUpdate(FUpdateTraversalContext& Context, const TTraitBinding<IUpdate>& Binding, const FTraitUpdateState& TraitState) const
-{
-	// Check if a new asset has been pushed 
-	const FSharedData* SharedData = Binding.GetSharedData<FSharedData>();
-	FInstanceData* InstanceData = Binding.GetInstanceData<FInstanceData>();
-		
-	const TObjectPtr<UObject> SharedDataAsset = SharedData->GetAsset(Binding);
-	if (SharedDataAsset != InstanceData->CachedAsset)
+	void FLayerAssetDataProviderTrait::PreUpdate(FUpdateTraversalContext& Context, const TTraitBinding<IUpdate>& Binding, const FTraitUpdateState& TraitState) const
 	{
-		TTraitBinding<IBlendStack> BlendStackBinding;
-		if(Binding.GetStackInterface<IBlendStack>(BlendStackBinding))
+		// Check if a new asset has been pushed 
+		const FSharedData* SharedData = Binding.GetSharedData<FSharedData>();
+		FInstanceData* InstanceData = Binding.GetInstanceData<FInstanceData>();
+			
+		const FGraphAssetHandle& SharedDataAssetHandle = SharedData->GetGraphAssetHandle(Binding);
+		if (SharedDataAssetHandle != InstanceData->CachedAssetHandle)
 		{
-			InstanceData->CachedAsset = SharedDataAsset;
-			if (SharedDataAsset != nullptr)
+			TTraitBinding<IBlendStack> BlendStackBinding;
+			if(Binding.GetStackInterface<IBlendStack>(BlendStackBinding))
 			{
-				FGraphAssetHandle AssetHandle = FAssetDataFactory::CreateUAFAssetDataFromObject<FUAFGraphFactoryAsset>(SharedDataAsset);
-				FAnimNextFactoryParams FactoryParams = FAnimGraphFactory::GetDefaultParamsForAsset(AssetHandle);
-		
-				const UUAFAnimGraph* AnimationGraph = IGraphFactory::GetOrBuildGraph(Context, Binding, AssetHandle, FactoryParams);
-				if(AnimationGraph)
+				if (SharedDataAssetHandle.IsValid())
 				{
-					IBlendStack::FGraphRequest NewGraphRequest;
-					NewGraphRequest.BlendArgs = FAlphaBlendArgs(0.0f);
-					NewGraphRequest.FactoryObject = SharedDataAsset;
-					NewGraphRequest.AnimationGraph = AnimationGraph;
-					NewGraphRequest.FactoryParams =  MoveTemp(FactoryParams);;
+					InstanceData->CachedAssetHandle = SharedDataAssetHandle;
+					
+					TArray<const UObject*> ReferencedObjects;
+					SharedDataAssetHandle.GetPtr()->GetObjectReferences(ReferencedObjects);
+					if(ReferencedObjects.Num() > 0)
+					{
+						FAnimNextFactoryParams FactoryParams = FAnimGraphFactory::GetDefaultParamsForAsset(SharedDataAssetHandle);
+						const UUAFAnimGraph* AnimationGraph = IGraphFactory::GetOrBuildGraph(Context, Binding, SharedDataAssetHandle, FactoryParams);
+						if(AnimationGraph)
+						{
+							IBlendStack::FGraphRequest NewGraphRequest;
+							NewGraphRequest.BlendArgs = FAlphaBlendArgs(0.0f);
+							NewGraphRequest.FactoryObject = ReferencedObjects[0];
+							NewGraphRequest.AnimationGraph = AnimationGraph;
+							NewGraphRequest.FactoryParams =  MoveTemp(FactoryParams);;
 
-					BlendStackBinding.PushGraph(Context, MoveTemp(NewGraphRequest));
-				}
-				else
-				{
-					UAF_TRAIT_LOG(Error, TEXT("Failed to create graph!"));
+							BlendStackBinding.PushGraph(Context, MoveTemp(NewGraphRequest));
+						}
+						else
+						{
+							UAF_TRAIT_LOG(Error, TEXT("Failed to create graph!"));
+						}
+					}
 				}
 			}
+			else
+			{
+				UAF_TRAIT_LOG(Error, TEXT("Blend Stack Trait is required but not found in stack!"));
+			}
 		}
-		else
-		{
-			UAF_TRAIT_LOG(Error, TEXT("Blend Stack Trait is required but not found in stack!"));
-		}
+			
+		IUpdate::PreUpdate(Context, Binding, TraitState);
 	}
-		
-	IUpdate::PreUpdate(Context, Binding, TraitState);
-}
 
-void FLayerAssetDataProviderTrait::AddReferencedObjects(const FExecutionContext& Context, const TTraitBinding<IGarbageCollection>& Binding, FReferenceCollector& Collector) const
-{
-	IGarbageCollection::AddReferencedObjects(Context, Binding, Collector);
-		
-	FInstanceData* InstanceData = Binding.GetInstanceData<FInstanceData>();
-	if(InstanceData->CachedAsset != nullptr)
+	void FLayerAssetDataProviderTrait::AddReferencedObjects(const FExecutionContext& Context, const TTraitBinding<IGarbageCollection>& Binding, FReferenceCollector& Collector) const
 	{
-		Collector.AddReferencedObject(InstanceData->CachedAsset);
+		IGarbageCollection::AddReferencedObjects(Context, Binding, Collector);
+		
+		const FInstanceData* InstanceData = Binding.GetInstanceData<FInstanceData>();
+		if(InstanceData->CachedAssetHandle.IsValid())
+		{
+			FGraphAssetHandle& AssetData = const_cast<FGraphAssetHandle&>(InstanceData->CachedAssetHandle);
+			TWeakObjectPtr<const UScriptStruct> ScriptStruct = AssetData.GetScriptStruct();
+			Collector.AddReferencedObjects(ScriptStruct, AssetData.GetMutableMemory());
+		}
 	}
-}
 	
 }
diff --git a/Engine/Plugins/Experimental/UAF/UAFLayering/Source/UAFLayering/Private/Traits/LayerAssetDataProviderTrait.h b/Engine/Plugins/Experimental/UAF/UAFLayering/Source/UAFLayering/Private/Traits/LayerAssetDataProviderTrait.h
--- a/Engine/Plugins/Experimental/UAF/UAFLayering/Source/UAFLayering/Private/Traits/LayerAssetDataProviderTrait.h
+++ b/Engine/Plugins/Experimental/UAF/UAFLayering/Source/UAFLayering/Private/Traits/LayerAssetDataProviderTrait.h
@@ -16,7 +16,7 @@ struct FLayerAssetDataProviderTrait : FAdditiveTrait, IUpdate, IGarbageCollectio
 
 	struct FInstanceData : FTrait::FInstanceData
 	{
-		TObjectPtr<UObject> CachedAsset = nullptr;
+		FGraphAssetHandle CachedAssetHandle;
 		void Construct(const FExecutionContext& Context, const FTraitBinding& Binding);
 	};
 	
diff --git a/Engine/Plugins/Experimental/UAF/UAFLayering/Source/UAFLayeringUncookedOnly/Internal/Layers/UAFLayerAssetProvider.h b/Engine/Plugins/Experimental/UAF/UAFLayering/Source/UAFLayeringUncookedOnly/Internal/Layers/UAFLayerAssetProvider.h
--- a/Engine/Plugins/Experimental/UAF/UAFLayering/Source/UAFLayeringUncookedOnly/Internal/Layers/UAFLayerAssetProvider.h
+++ b/Engine/Plugins/Experimental/UAF/UAFLayering/Source/UAFLayeringUncookedOnly/Internal/Layers/UAFLayerAssetProvider.h
@@ -5,6 +5,8 @@
 #include "CoreMinimal.h"
 #include "UAFLayerContentProviderBase.h"
 #include "Factory/AnimNextFactoryParams.h"
+#include "UAF/UAFAssetData.h"
+
 #include "UAFLayerAssetProvider.generated.h"
 
 class UUAFLayer;
@@ -31,21 +33,41 @@ public:
 	virtual URigVMPin* CreateLayerContentTrait(UE::UAF::Layering::FLayerCreationContext& LayerCreationContext) override;
 	virtual TSharedRef<SWidget> CreateLayerContentWidget(UUAFLayer* InLayer) override;
 
-	void SetLayerAsset(UObject* InNewAsset);
+	void SetLayerAsset(const UObject* InNewAsset);
+	
+#if WITH_EDITORONLY_DATA
+	void PostSerialize(const FArchive& Ar);
+#endif
 	
 private:
 	void OnLayerAssetSelectionChanged(const FAssetData& InAssetData);
 	
 private:
-	UPROPERTY(EditAnywhere, Category = "Layer|Content", meta = (GetAllowedClasses = "/Script/UAFAnimGraph.AnimNextAnimGraphSettings:GetAllowedAssetClasses"))
-	TObjectPtr<UObject> LayerAsset = nullptr;
+	UPROPERTY(EditAnywhere, Category = "Layer|Content")
+	TInstancedStruct<FUAFGraphFactoryAsset> AssetData; 
+	
+	UPROPERTY()
+	TWeakObjectPtr<UUAFLayer> OuterLayer = nullptr;
 
 #if WITH_EDITORONLY_DATA
 	UE_DEPRECATED(5.8, "Factory Parameters are not used directly anymore, instead it uses FUAFAssetData internally")
 	UPROPERTY()
 	FAnimNextFactoryParams Parameters_DEPRECATED;
-#endif
 	
+	UE_DEPRECATED(5.8, "Using UObject based assets directly is deprecated, use AssetData instead")
 	UPROPERTY()
-	TWeakObjectPtr<UUAFLayer> OuterLayer = nullptr;
+	TObjectPtr<UObject> LayerAsset_DEPRECATED = nullptr;
+#endif
+	
 };
+
+template<>
+struct TStructOpsTypeTraits<FUAFLayerAssetProvider> : public TStructOpsTypeTraitsBase2<FUAFLayerAssetProvider>
+{
+#if WITH_EDITORONLY_DATA
+	enum
+	{
+		WithPostSerialize = true,
+	};
+#endif
+};
\ No newline at end of file
diff --git a/Engine/Plugins/Experimental/UAF/UAFLayering/Source/UAFLayeringUncookedOnly/Private/Layers/UAFLayerAssetProvider.cpp b/Engine/Plugins/Experimental/UAF/UAFLayering/Source/UAFLayeringUncookedOnly/Private/Layers/UAFLayerAssetProvider.cpp
--- a/Engine/Plugins/Experimental/UAF/UAFLayering/Source/UAFLayeringUncookedOnly/Private/Layers/UAFLayerAssetProvider.cpp
+++ b/Engine/Plugins/Experimental/UAF/UAFLayering/Source/UAFLayeringUncookedOnly/Private/Layers/UAFLayerAssetProvider.cpp
@@ -17,7 +17,7 @@
 #include "TraitCore/Trait.h"
 #include "Traits/BlendStackTrait.h"
 #include "Traits/LayerAssetDataTraitData.h"
-#include "Traits/SubGraphHostTraitData.h"
+#include "UAF/UAFAssetFactory.h"
 #include "Widgets/Layout/SBox.h"
 
 #define LOCTEXT_NAMESPACE "UAFLayerAssetContentProvider"
@@ -32,48 +32,30 @@ URigVMPin* FUAFLayerAssetProvider::CreateLayerContentTrait(UE::UAF::Layering::FL
 		return nullptr;
 	}
 	
-	// TODO: We likely could unify this approach and not differentiate between graphs and other assets
-	if (UUAFAnimGraph* AnimationGraph = Cast<UUAFAnimGraph>(LayerAsset))
+	// Create blend stack trait
+	const FName BlendStackTraitName = LayerCreationContext.GraphController->AddTraitStruct(ContentHostTraitStackNode, TInstancedStruct<FAnimNextBlendStackCoreTraitSharedData>::Make(), INDEX_NONE, false, false);
+	if (BlendStackTraitName.IsNone())
 	{
-		FAnimNextSubGraphHostTraitSharedData DefaultSubGraphData;
-		DefaultSubGraphData.AnimationGraph = AnimationGraph;
-		
-		const FName SubGraphHostTraitName = LayerCreationContext.GraphController->AddTraitStruct(ContentHostTraitStackNode, TInstancedStruct<FAnimNextSubGraphHostTraitSharedData>::Make(DefaultSubGraphData), INDEX_NONE, false, false);
-		if (SubGraphHostTraitName.IsNone())
-		{
-			LayerCreationContext.CompileSettings.ReportError(TEXT("Failed to create SubGraphHost Trait."));
-			return nullptr;
-		}
-		
-		// Expand trait pin per default 
-		URigVMPin* TraitPin = ContentHostTraitStackNode->FindTrait(SubGraphHostTraitName);
-		LayerCreationContext.GraphController->SetPinExpansion(TraitPin->GetPinPath(), true, false);
+		LayerCreationContext.CompileSettings.ReportError(TEXT("Failed to create Blend Stack Trait."));
+		return nullptr;
 	}
-	else
+	
+	// Create asset data provider trait 
+	FUAFLayerAssetDataTraitSharedData AssetDataDefaults;
+	AssetDataDefaults.GraphAssetHandle = AssetData;
+	
+	const FName AssetDataTraitName = LayerCreationContext.GraphController->AddTraitStruct(ContentHostTraitStackNode, TInstancedStruct<FUAFLayerAssetDataTraitSharedData>::Make(AssetDataDefaults), INDEX_NONE, false, false);
+	if (AssetDataTraitName.IsNone())
 	{
-		// Create Blend stack trait
-		const FName BlendStackTraitName = LayerCreationContext.GraphController->AddTraitStruct(ContentHostTraitStackNode, TInstancedStruct<FAnimNextBlendStackCoreTraitSharedData>::Make(), INDEX_NONE, false, false);
-		if (BlendStackTraitName.IsNone())
-		{
-			LayerCreationContext.CompileSettings.ReportError(TEXT("Failed to create Blend Stack Trait."));
-			return nullptr;
-		}
-		
-		// Create asset data provider trait 
-		FUAFLayerAssetDataTraitSharedData AssetDataDefaults;
-		AssetDataDefaults.Asset = LayerAsset;
-		
-		const FName AssetDataTraitName = LayerCreationContext.GraphController->AddTraitStruct(ContentHostTraitStackNode, TInstancedStruct<FUAFLayerAssetDataTraitSharedData>::Make(AssetDataDefaults), INDEX_NONE, false, false);
-		if (AssetDataTraitName.IsNone())
-		{
-			LayerCreationContext.CompileSettings.ReportError(TEXT("Failed to create Layer Asset DataTrait."));
-			return nullptr;
-		}
-		
-		URigVMPin* TraitPin = ContentHostTraitStackNode->FindTrait(AssetDataTraitName);
-		LayerCreationContext.GraphController->SetPinExpansion(TraitPin->GetPinPath(), true, false);
+		LayerCreationContext.CompileSettings.ReportError(TEXT("Failed to create LayerAssetData Trait."));
+		return nullptr;
 	}
-
+	
+	// Expand trait pin per default
+	URigVMPin* TraitPin = ContentHostTraitStackNode->FindTrait(AssetDataTraitName);
+	LayerCreationContext.GraphController->SetPinExpansion(TraitPin->GetPinPath(), true, false);
+	
+	// Return the result pin of the trait stack to allow building of the graph
 	URigVMPin* BlendTraitStackResult = ContentHostTraitStackNode->FindPin(GET_MEMBER_NAME_STRING_CHECKED(FRigUnit_AnimNextTraitStack, Result));
 	if (BlendTraitStackResult == nullptr)
 	{
@@ -84,15 +66,22 @@ URigVMPin* FUAFLayerAssetProvider::CreateLayerContentTrait(UE::UAF::Layering::FL
 	return BlendTraitStackResult;
 }
 
-void FUAFLayerAssetProvider::SetLayerAsset(UObject* InNewAsset)
+void FUAFLayerAssetProvider::SetLayerAsset(const UObject* InNewAsset)
 {
 	const FScopedTransaction Transaction(LOCTEXT("OnLayerAssetSelectionChangedAction", "Set Layer Asset"));
 	if (OuterLayer.IsValid())
 	{
 		OuterLayer->Modify();
 	}
-	
-	LayerAsset = InNewAsset;
+
+	if (InNewAsset)
+	{
+		AssetData = UE::UAF::FAssetDataFactory::CreateUAFAssetDataFromObject<FUAFGraphFactoryAsset>(InNewAsset);
+	}
+	else
+	{
+		AssetData.Reset();
+	}
 	
 	if (OuterLayer.IsValid())
 	{
@@ -125,9 +114,12 @@ TSharedRef<SWidget> FUAFLayerAssetProvider::CreateLayerContentWidget(UUAFLayer*
 					})
 				.ObjectPath_Lambda([this]()
 					{
-						if (LayerAsset)
+						if (AssetData.IsValid())
 						{
-							return LayerAsset->GetPathName();
+							TArray<const UObject*> OutReferencedObjects;
+							AssetData.GetPtr()->GetObjectReferences(OutReferencedObjects);
+							
+							return OutReferencedObjects.Num() > 0 ? OutReferencedObjects[0]->GetPathName() : FString();
 						}
 						return FString();
 					})
@@ -137,5 +129,20 @@ TSharedRef<SWidget> FUAFLayerAssetProvider::CreateLayerContentWidget(UUAFLayer*
 		];
 }
 
+#if WITH_EDITORONLY_DATA
+void FUAFLayerAssetProvider::PostSerialize(const FArchive& Ar)
+{
+	if (Ar.IsLoading())
+	{
+		PRAGMA_DISABLE_DEPRECATION_WARNINGS
+		if (LayerAsset_DEPRECATED != nullptr && AssetData.IsValid() == false)
+		{
+			AssetData = UE::UAF::FAssetDataFactory::CreateUAFAssetDataFromObject<FUAFGraphFactoryAsset>(LayerAsset_DEPRECATED);
+		}
+		PRAGMA_ENABLE_DEPRECATION_WARNINGS
+	}
+}
+#endif
+
 #undef LOCTEXT_NAMESPACE //UAFLayerAssetContentProvider
 
