UAF Layering: Montage Layer Support

- New layer type to automatically play montages based on slot name
- Option to automatically enable layer if a montage is playing
- Option to automatically set blend times based on currently playing montage

#Synced-CL 50858266

diff --git a/Engine/Plugins/Experimental/UAF/UAFAnimGraph/Source/UAFAnimGraph/Internal/Injection/UAFMontageComponent.h b/Engine/Plugins/Experimental/UAF/UAFAnimGraph/Source/UAFAnimGraph/Internal/Injection/UAFMontageComponent.h
--- a/Engine/Plugins/Experimental/UAF/UAFAnimGraph/Source/UAFAnimGraph/Internal/Injection/UAFMontageComponent.h
+++ b/Engine/Plugins/Experimental/UAF/UAFAnimGraph/Source/UAFAnimGraph/Internal/Injection/UAFMontageComponent.h
@@ -12,6 +12,8 @@
 #include "Module/UAFModuleInstanceComponent.h"
 #include "UAFMontageComponent.generated.h"
 
+#define UE_API UAFANIMGRAPH_API
+
 struct FUAFUpdateMontageTickFunction : public FTickFunction
 {
 public:
@@ -81,7 +83,9 @@ public:
 	FUAFMontageComponent() = default;
 	virtual ~FUAFMontageComponent() override;
 
-	bool IsSlotActive(FName SlotNodeName) const;
+	UE_API bool IsSlotActive(FName SlotNodeName) const;
+	UE_API TObjectPtr<const UAnimMontage> GetActiveMontageForSlot(FName SlotNodeName) const;
+	
 private:
 	virtual void OnBindToInstance() override;
 
@@ -151,3 +155,4 @@ struct FRigUnit_IsSlotActive : public FRigUnit_AnimNextBase
 	bool Result = false;
 };
 
+#undef UE_API
diff --git a/Engine/Plugins/Experimental/UAF/UAFAnimGraph/Source/UAFAnimGraph/Private/Injection/UAFMontageComponent.cpp b/Engine/Plugins/Experimental/UAF/UAFAnimGraph/Source/UAFAnimGraph/Private/Injection/UAFMontageComponent.cpp
--- a/Engine/Plugins/Experimental/UAF/UAFAnimGraph/Source/UAFAnimGraph/Private/Injection/UAFMontageComponent.cpp
+++ b/Engine/Plugins/Experimental/UAF/UAFAnimGraph/Source/UAFAnimGraph/Private/Injection/UAFMontageComponent.cpp
@@ -127,6 +127,21 @@ bool FUAFMontageComponent::IsSlotActive(FName SlotNodeName) const
 	return false;
 }
 
+TObjectPtr<const UAnimMontage> FUAFMontageComponent::GetActiveMontageForSlot(FName SlotNodeName) const
+{
+	if (SlotNodeName != NAME_None)
+	{
+		for (const FUAFMontageEvalData& EvalData : MontageEvalData)
+		{
+			if (EvalData.IsActive() && EvalData.GetMontage()->IsValidSlot(SlotNodeName))
+			{
+				return EvalData.GetMontage();
+			}
+		}
+	}
+	return nullptr;
+}
+
 bool FUAFMontageComponent::GetActiveNotifiesForSlot(FName SlotName, TArray<FAnimNotifyEventReference>& OutNotifies) const
 {
 	const FAnimNotifyArray* FoundEntry = ActiveSlotNotifies.Find(SlotName);
diff --git a/Engine/Plugins/Experimental/UAF/UAFLayering/Source/UAFLayering/Internal/Traits/MontageLayerTraitData.h b/Engine/Plugins/Experimental/UAF/UAFLayering/Source/UAFLayering/Internal/Traits/MontageLayerTraitData.h
new file mode 100644
--- /dev/null
+++ b/Engine/Plugins/Experimental/UAF/UAFLayering/Source/UAFLayering/Internal/Traits/MontageLayerTraitData.h
@@ -0,0 +1,33 @@
+// Copyright Epic Games, Inc. All Rights Reserved.
+
+#pragma once
+
+#include "TraitCore/TraitSharedData.h"
+
+#include "MontageLayerTraitData.generated.h"
+
+
+/** A trait that supports modifying an existing layer based on montage runtime data */
+USTRUCT(meta = (DisplayName = "Montage Layer Data", ShowTooltip=true, Hidden))
+struct FUAFMontageLayerTraitSharedData : public FAnimNextTraitSharedData
+{
+	GENERATED_BODY()
+	
+	UPROPERTY(EditAnywhere, Category = "Layer", meta = (Inline))
+	bool bAutoEnableLayer = false;
+	
+	UPROPERTY(EditAnywhere, Category = "Layer", meta = (Inline))
+	bool bAutoSetBlendTimes = false;
+	
+	// The montage slot name of this layer 
+	UPROPERTY(EditAnywhere, Category = "Layer", meta = (Inline))
+	FName SlotName = NAME_None;
+	
+	// The layer name of this layer 
+	UPROPERTY(EditAnywhere, Category = "Layer", meta = (Inline))
+	FName LayerName = NAME_None;
+	
+	// The soft object path to the owning layer stack
+	UPROPERTY(EditAnywhere, Category = "Layer", meta = (Inline))
+	FString LayerStackPath;
+};
diff --git a/Engine/Plugins/Experimental/UAF/UAFLayering/Source/UAFLayering/Private/Traits/MontageLayerDataTrait.cpp b/Engine/Plugins/Experimental/UAF/UAFLayering/Source/UAFLayering/Private/Traits/MontageLayerDataTrait.cpp
new file mode 100644
--- /dev/null
+++ b/Engine/Plugins/Experimental/UAF/UAFLayering/Source/UAFLayering/Private/Traits/MontageLayerDataTrait.cpp
@@ -0,0 +1,75 @@
+// Copyright Epic Games, Inc. All Rights Reserved.
+
+#include "Traits/MontageLayerDataTrait.h"
+
+#include "UAFLayeringTypes.h"
+#include "Injection/UAFMontageComponent.h"
+#include "Module/AnimNextModuleInstance.h"
+
+namespace UE::UAF
+{
+	AUTO_REGISTER_ANIM_TRAIT(FMontageLayerDataTrait)
+
+	#define TRAIT_INTERFACE_ENUMERATOR(GeneratorMacro) \
+		GeneratorMacro(IUpdate) \
+	
+	GENERATE_ANIM_TRAIT_IMPLEMENTATION(FMontageLayerDataTrait, TRAIT_INTERFACE_ENUMERATOR, NULL_ANIM_TRAIT_INTERFACE_ENUMERATOR, NULL_ANIM_TRAIT_EVENT_ENUMERATOR)
+	#undef TRAIT_INTERFACE_ENUMERATOR
+	#undef TRAIT_EVENT_ENUMERATOR
+	
+
+	void FMontageLayerDataTrait::PreUpdate(FUpdateTraversalContext& Context, const TTraitBinding<IUpdate>& Binding, const FTraitUpdateState& TraitState) const
+	{
+		const FSharedData* SharedData = Binding.GetSharedData<FSharedData>();
+		
+		const bool bAutoEnableLayer = SharedData->bAutoEnableLayer;
+		const bool bAutoSetBlendTimes = SharedData->bAutoSetBlendTimes;
+		if (bAutoEnableLayer || bAutoSetBlendTimes)
+		{
+			if (FAnimNextModuleInstance* ModuleInstance = Context.GetRootGraphInstance().GetModuleInstance())
+			{
+				FInstanceData* InstanceData = Binding.GetInstanceData<FInstanceData>();
+				FUAFMontageComponent& MontageComponent = ModuleInstance->GetComponent<FUAFMontageComponent>();
+				const bool bSlotActive = MontageComponent.IsSlotActive(SharedData->SlotName);
+				
+				// Enable or disable the layer based on if a montage is currently playing in the relevant slot 
+				if (bAutoEnableLayer && (bSlotActive != InstanceData->bLastEnableState || TraitState.IsNewlyRelevant()))
+				{
+					InstanceData->bLastEnableState = bSlotActive;
+				
+					TSharedPtr<Layering::FLayerStack_LayerEvent> Event = MakeTraitEvent<Layering::FLayerStack_LayerEvent>(SharedData->LayerName, SharedData->LayerStackPath);
+					Event->Action = bSlotActive ? Layering::ELayerEventAction::EnableLayer : Layering::ELayerEventAction::DisableLayer;
+					Context.RaiseOutputTraitEvent(Event);
+				}
+
+				if (bAutoSetBlendTimes && bSlotActive)
+				{
+					const TObjectPtr<const UAnimMontage> PlayingMontage = MontageComponent.GetActiveMontageForSlot(SharedData->SlotName);
+					if (PlayingMontage && InstanceData->LastActiveMontage != PlayingMontage)
+					{
+						InstanceData->LastActiveMontage = PlayingMontage;
+						
+						TSharedPtr<Layering::FLayerStack_LayerEvent> SetBlendInTimeEvent = MakeTraitEvent<Layering::FLayerStack_LayerEvent>(SharedData->LayerName, SharedData->LayerStackPath);
+						SetBlendInTimeEvent->Action = Layering::ELayerEventAction::SetFloatValue;
+						SetBlendInTimeEvent->PropertyToSet = Layering::FLayerStack_LayerEvent::BlendInTimeProperty;
+						SetBlendInTimeEvent->FloatValue = PlayingMontage->BlendIn.GetBlendTime();
+						Context.RaiseOutputTraitEvent(SetBlendInTimeEvent);
+						
+						TSharedPtr<Layering::FLayerStack_LayerEvent> SetBlendOutTimeEvent = MakeTraitEvent<Layering::FLayerStack_LayerEvent>(SharedData->LayerName, SharedData->LayerStackPath);
+						SetBlendOutTimeEvent->Action = Layering::ELayerEventAction::SetFloatValue;
+						SetBlendOutTimeEvent->PropertyToSet = Layering::FLayerStack_LayerEvent::BlendOutTimeProperty;
+						SetBlendOutTimeEvent->FloatValue = PlayingMontage->BlendOut.GetBlendTime();
+						Context.RaiseOutputTraitEvent(SetBlendOutTimeEvent);
+					}
+				}
+				
+				if (!bSlotActive && InstanceData->LastActiveMontage.IsValid())
+				{
+					InstanceData->LastActiveMontage = nullptr;
+				}
+			}
+		}
+		
+		IUpdate::PreUpdate(Context, Binding, TraitState);
+	}
+}
diff --git a/Engine/Plugins/Experimental/UAF/UAFLayering/Source/UAFLayering/Private/Traits/MontageLayerDataTrait.h b/Engine/Plugins/Experimental/UAF/UAFLayering/Source/UAFLayering/Private/Traits/MontageLayerDataTrait.h
new file mode 100644
--- /dev/null
+++ b/Engine/Plugins/Experimental/UAF/UAFLayering/Source/UAFLayering/Private/Traits/MontageLayerDataTrait.h
@@ -0,0 +1,32 @@
+// Copyright Epic Games, Inc. All Rights Reserved.
+
+#pragma once
+
+#include "TraitCore/Trait.h"
+#include "TraitInterfaces/IUpdate.h"
+#include "Traits/MontageLayerTraitData.h"
+
+namespace UE::UAF
+{
+	struct FMontageLayerDataTrait : FAdditiveTrait, IUpdate
+	{
+		DECLARE_ANIM_TRAIT(FMontageLayerDataTrait, FAdditiveTrait)
+
+		struct FInstanceData : FTrait::FInstanceData
+		{
+			bool bLastEnableState = false;
+			TWeakObjectPtr<const UAnimMontage> LastActiveMontage = nullptr;
+		};
+	
+		using FSharedData = FUAFMontageLayerTraitSharedData;
+	
+		// IUpdate impl
+		virtual void PreUpdate(FUpdateTraversalContext& Context, const TTraitBinding<IUpdate>& Binding, const FTraitUpdateState& TraitState) const override;
+
+#if WITH_EDITOR
+		virtual bool IsHidden() const override { return true; };
+#endif
+	
+	};
+	
+}
diff --git a/Engine/Plugins/Experimental/UAF/UAFLayering/Source/UAFLayeringUncookedOnly/Internal/Layers/UAFLayer.h b/Engine/Plugins/Experimental/UAF/UAFLayering/Source/UAFLayeringUncookedOnly/Internal/Layers/UAFLayer.h
--- a/Engine/Plugins/Experimental/UAF/UAFLayering/Source/UAFLayeringUncookedOnly/Internal/Layers/UAFLayer.h
+++ b/Engine/Plugins/Experimental/UAF/UAFLayering/Source/UAFLayeringUncookedOnly/Internal/Layers/UAFLayer.h
@@ -47,11 +47,11 @@ public:
 	EUAFLayerState GetLayerState() const;
 	void SetLayerState(const EUAFLayerState InLayerState);
 	
-	bool IsLayerContentProviderValid() const;
+	bool IsLayerBlendProviderValid() const;
 	TInstancedStruct<FUAFLayerBlendProviderBase>& GetLayerBlendProvider();
 	void SetLayerBlendProvider(const TInstancedStruct<FUAFLayerBlendProviderBase>& InLayerBlendProvider);
 	
-	bool IsLayerBlendProviderValid() const;
+	bool IsLayerContentProviderValid() const;
 	TInstancedStruct<FUAFLayerContentProviderBase>& GetLayerContentProvider();
 	void SetLayerContentProvider(const TInstancedStruct<FUAFLayerContentProviderBase>& InLayerContentProvider);
 	
diff --git a/Engine/Plugins/Experimental/UAF/UAFLayering/Source/UAFLayeringUncookedOnly/Internal/Layers/UAFLayerBlendProviderBase.h b/Engine/Plugins/Experimental/UAF/UAFLayering/Source/UAFLayeringUncookedOnly/Internal/Layers/UAFLayerBlendProviderBase.h
--- a/Engine/Plugins/Experimental/UAF/UAFLayering/Source/UAFLayeringUncookedOnly/Internal/Layers/UAFLayerBlendProviderBase.h
+++ b/Engine/Plugins/Experimental/UAF/UAFLayering/Source/UAFLayeringUncookedOnly/Internal/Layers/UAFLayerBlendProviderBase.h
@@ -23,7 +23,7 @@ public:
 	
 	// Responsible for creating traits to represent this blend and link inputs to the relevant traits.
 	// @return The final output pin of this layer blend which can be linked as input by other layers or be linked to the result of the graph
-	virtual URigVMPin* CreateBlendGraphTrait(UE::UAF::Layering::FLayerCreationContext& LayerCreationContext) {return nullptr;};
+	virtual URigVMPin* CreateBlendGraphTrait(UE::UAF::Layering::FLayerCreationContext& LayerCreationContext) { return nullptr; };
 	
 	// Creates the widget to represent this blend in the layer stack editor 
 	virtual TSharedRef<SWidget> CreateLayerBlendWidget(UUAFLayer* OuterLayer) { return SNullWidget::NullWidget;} ;
@@ -32,7 +32,12 @@ public:
 	virtual const FSlateBrush* GetOverrideLayerBackground() const { return nullptr; };
 	
 	// Optional override of the containing layers indicator color
-	virtual bool GetOverrideIndicatorColor(FSlateColor& OutSlateColor) const {return false;};
+	virtual bool GetOverrideIndicatorColor(FSlateColor& OutSlateColor) const { return false; };
+	
+	void SetAlwaysUpdateChildren(bool bInAlwaysUpdate) { bAlwaysUpdateChildren = bInAlwaysUpdate; };
+	bool GetAlwaysUpdateChildren() const { return bAlwaysUpdateChildren; };
+private:
+	bool bAlwaysUpdateChildren = false;	
 };
 
 
diff --git a/Engine/Plugins/Experimental/UAF/UAFLayering/Source/UAFLayeringUncookedOnly/Internal/Layers/UAFMontageProvider.h b/Engine/Plugins/Experimental/UAF/UAFLayering/Source/UAFLayeringUncookedOnly/Internal/Layers/UAFMontageProvider.h
new file mode 100644
--- /dev/null
+++ b/Engine/Plugins/Experimental/UAF/UAFLayering/Source/UAFLayeringUncookedOnly/Internal/Layers/UAFMontageProvider.h
@@ -0,0 +1,46 @@
+// Copyright Epic Games, Inc. All Rights Reserved.
+
+#pragma once
+
+#include "CoreMinimal.h"
+#include "UAFLayerContentProviderBase.h"
+
+#include "UAFMontageProvider.generated.h"
+
+class UUAFLayer;
+class SWidget;
+
+// Layer content provider to support legacy montages
+USTRUCT(DisplayName="Montage Layer")
+struct FUAFMontageProvider : public FUAFLayerContentProviderBase
+{
+	GENERATED_BODY()
+
+public:
+	virtual URigVMPin* CreateLayerContentTrait(UE::UAF::Layering::FLayerCreationContext& LayerCreationContext) override;
+	virtual TSharedRef<SWidget> CreateLayerContentWidget(UUAFLayer* InLayer) override;
+	
+private:
+	// The name of this slot 
+	UPROPERTY(EditAnywhere, Category = "Layer|Content")
+	FName SlotName = NAME_None;
+	
+	// If true it will always update the source input, even when montages are fully blended in 
+	UPROPERTY(EditAnywhere, Category = "Layer|Content")
+	bool bAlwaysUpdateSource = false;
+	
+	// Disables synchronization for this slot even if the playing montage has a sync group setup 
+	UPROPERTY(EditAnywhere, Category = "Layer|Content")
+	bool bDisableSynchronization = false;
+	
+	// If true, automatically disables and enables the layer based on if a montage in the given slot is active
+	// This will cause the content of this layer to always update even if the weight is 0 to auto enable itself 
+	UPROPERTY(EditAnywhere, Category = "Layer|Content")
+	bool bAutoEnableLayerWithMontage = false;
+	
+	// If true, automatically sets the layer blend in and blend out time based on the montage playing 
+	UPROPERTY(EditAnywhere, Category = "Layer|Content")
+	bool bAutoSetBlendTimesFromMontage = false;
+	
+};
+
diff --git a/Engine/Plugins/Experimental/UAF/UAFLayering/Source/UAFLayeringUncookedOnly/Private/Layers/UAFLayerDefaultBlendProvider.cpp b/Engine/Plugins/Experimental/UAF/UAFLayering/Source/UAFLayeringUncookedOnly/Private/Layers/UAFLayerDefaultBlendProvider.cpp
--- a/Engine/Plugins/Experimental/UAF/UAFLayering/Source/UAFLayeringUncookedOnly/Private/Layers/UAFLayerDefaultBlendProvider.cpp
+++ b/Engine/Plugins/Experimental/UAF/UAFLayering/Source/UAFLayeringUncookedOnly/Private/Layers/UAFLayerDefaultBlendProvider.cpp
@@ -151,6 +151,7 @@ URigVMPin* FUAFDefaultBlendProvider::CreateBlendGraphTrait(UE::UAF::Layering::FL
 			FUAFBlendLayerTraitSharedData DefaultTraitData; 
 			DefaultTraitData.BlendWeight = LayerWeight;
 			DefaultTraitData.BlendProfile = BlendMask;
+			DefaultTraitData.bAlwaysUpdateChildBlend = GetAlwaysUpdateChildren();
 			
 			// Create trait 
 			BlendBaseTraitName = LayerCreationContext.GraphController->AddTraitStruct(BlendTraitStackNode, TInstancedStruct<FUAFBlendLayerTraitSharedData>::Make(DefaultTraitData), INDEX_NONE, false, false);
@@ -164,6 +165,7 @@ URigVMPin* FUAFDefaultBlendProvider::CreateBlendGraphTrait(UE::UAF::Layering::FL
 			// Setup default trait values 
 			FAnimNextBlendTwoWayTraitSharedData DefaultTraitData; 
 			DefaultTraitData.BlendWeight = LayerWeight;
+			DefaultTraitData.bAlwaysUpdateChildren = GetAlwaysUpdateChildren();
 			
 			// Create two-way blend trait 
 			BlendBaseTraitName = LayerCreationContext.GraphController->AddTraitStruct(BlendTraitStackNode, TInstancedStruct<FAnimNextBlendTwoWayTraitSharedData>::Make(DefaultTraitData), INDEX_NONE, false, false);
diff --git a/Engine/Plugins/Experimental/UAF/UAFLayering/Source/UAFLayeringUncookedOnly/Private/Layers/UAFMontageProvider.cpp b/Engine/Plugins/Experimental/UAF/UAFLayering/Source/UAFLayeringUncookedOnly/Private/Layers/UAFMontageProvider.cpp
new file mode 100644
--- /dev/null
+++ b/Engine/Plugins/Experimental/UAF/UAFLayering/Source/UAFLayeringUncookedOnly/Private/Layers/UAFMontageProvider.cpp
@@ -0,0 +1,96 @@
+// Copyright Epic Games, Inc. All Rights Reserved.
+
+
+#include "Layers/UAFMontageProvider.h"
+
+#include "AnimNextController.h"
+#include "IDetailTreeNode.h"
+#include "UAFLayerStack.h"
+#include "UAFLayerStack_EditorData.h"
+#include "Components/HorizontalBox.h"
+#include "Graph/RigUnit_AnimNextTraitStack.h"
+#include "Injection/MontageTrait.h"
+#include "Layers/UAFLayer.h"
+#include "Traits/MontageLayerTraitData.h"
+#include "Widgets/Input/SEditableTextBox.h"
+#include "Widgets/Layout/SBox.h"
+
+#define LOCTEXT_NAMESPACE "UAFMontageProvider"
+
+URigVMPin* FUAFMontageProvider::CreateLayerContentTrait(UE::UAF::Layering::FLayerCreationContext& LayerCreationContext)
+{
+	LayerCreationContext.LastNodeLocation += FVector2D(0.0f, 500.0f);
+	
+	// Create trait stack
+	URigVMUnitNode* MontageTraitStackNode = LayerCreationContext.GraphController->AddUnitNode(FRigUnit_AnimNextTraitStack::StaticStruct(), FRigVMStruct::ExecuteName, LayerCreationContext.LastNodeLocation, FString(), false);
+	if (MontageTraitStackNode == nullptr)
+	{
+		LayerCreationContext.CompileSettings.ReportError(TEXT("Failed to create trait stack to host layer content."));
+		return nullptr;
+	}
+	
+	// Create legacy montage trait 
+	FUAFMontageTraitSharedData DataDefaults;
+	DataDefaults.SlotName = SlotName;
+	DataDefaults.bAlwaysUpdateSource = bAlwaysUpdateSource;
+	DataDefaults.bDisableSynchronization = bDisableSynchronization;
+	
+	const FName MontageTraitName = LayerCreationContext.GraphController->AddTraitStruct(MontageTraitStackNode, TInstancedStruct<FUAFMontageTraitSharedData>::Make(DataDefaults), INDEX_NONE, false, false);
+	if (MontageTraitName.IsNone())
+	{
+		LayerCreationContext.CompileSettings.ReportError(TEXT("Failed to create Montage Trait."));
+		return nullptr;
+	}
+	
+	// Create additive montage trait if required
+	if (bAutoEnableLayerWithMontage || bAutoSetBlendTimesFromMontage)
+	{
+		FUAFMontageLayerTraitSharedData LayerDataDefaults;
+		LayerDataDefaults.bAutoEnableLayer = bAutoEnableLayerWithMontage;
+		LayerDataDefaults.bAutoSetBlendTimes = bAutoSetBlendTimesFromMontage;
+		LayerDataDefaults.SlotName = SlotName;
+		LayerDataDefaults.LayerName = LayerCreationContext.Layer->GetLayerName();
+		LayerDataDefaults.LayerStackPath = FSoftObjectPath::ConstructFromObject(LayerCreationContext.LayerStack.Get()).ToString();
+	
+		const FName MontageLayerTraitName = LayerCreationContext.GraphController->AddTraitStruct(MontageTraitStackNode, TInstancedStruct<FUAFMontageLayerTraitSharedData>::Make(LayerDataDefaults), INDEX_NONE, false, false);
+		if (MontageLayerTraitName.IsNone())
+		{
+			LayerCreationContext.CompileSettings.ReportError(TEXT("Failed to create Montage Layer Trait."));
+			return nullptr;
+		}
+
+		if (bAutoEnableLayerWithMontage && LayerCreationContext.Layer->IsLayerBlendProviderValid())
+		{
+			// if we want to be able to modify the layers behavior even if it is currently not enabled (e.g. to auto enable it) we have to always update it
+			LayerCreationContext.Layer->GetLayerBlendProvider().GetMutablePtr()->SetAlwaysUpdateChildren(true);
+		}
+	}
+	
+	// return result pin to connect to the rest of the graph 
+	URigVMPin* BlendTraitStackResult = MontageTraitStackNode->FindPin(GET_MEMBER_NAME_STRING_CHECKED(FRigUnit_AnimNextTraitStack, Result));
+	if (BlendTraitStackResult == nullptr)
+	{
+		LayerCreationContext.CompileSettings.ReportError(TEXT("Failed to retrieve Result pin from trait stack"));
+		return nullptr;
+	}
+	
+	return BlendTraitStackResult;
+}
+
+TSharedRef<SWidget> FUAFMontageProvider::CreateLayerContentWidget(UUAFLayer* InLayer)
+{
+	return SNew(SHorizontalBox)
+		+ SHorizontalBox::Slot()
+		.VAlign(EVerticalAlignment::VAlign_Center)
+		[
+			SNew(SEditableTextBox)
+			.ToolTipText(LOCTEXT("SlotNameLabel", "The name of the montage slot"))
+			.Text(FText::FromName(SlotName))
+			.OnTextCommitted_Lambda([this](const FText& Val, ETextCommit::Type TextCommitType)
+			{
+				SlotName = FName(Val.ToString());
+			})
+		];
+}
+
+#undef LOCTEXT_NAMESPACE //UAFMontageProvider
