UAF Layering: Add workspace preview support for layering

#Synced-CL 50903586

diff --git a/Engine/Plugins/Experimental/UAF/UAFLayering/Source/UAFLayeringEditor/Private/UAFLayeringEditorModule.cpp b/Engine/Plugins/Experimental/UAF/UAFLayering/Source/UAFLayeringEditor/Private/UAFLayeringEditorModule.cpp
--- a/Engine/Plugins/Experimental/UAF/UAFLayering/Source/UAFLayeringEditor/Private/UAFLayeringEditorModule.cpp
+++ b/Engine/Plugins/Experimental/UAF/UAFLayering/Source/UAFLayeringEditor/Private/UAFLayeringEditorModule.cpp
@@ -14,6 +14,7 @@
 #include "Tools/EdModeInteractiveToolsContext.h"
 #include "Widgets/SLayerStack.h"
 #include "Workspace/LayerStackItemDetails.h"
+#include "Workspace/LayerStackViewportController.h"
 #include "Workspace/UAFLayerStackWorkspaceOutlinerData.h"
 
 
@@ -42,6 +43,11 @@ void FUAFLayeringEditorModule::StartupModule()
 	UE::Workspace::IWorkspaceEditorModule& WorkspaceEditorModule = FModuleManager::Get().LoadModuleChecked<UE::Workspace::IWorkspaceEditorModule>("WorkspaceEditor");
 	WorkspaceEditorModule.RegisterObjectDocumentType(UUAFLayerStack::StaticClass()->GetClassPathName(), LayerStackDocumentArgs);
 	WorkspaceEditorModule.RegisterWorkspaceItemDetails(UE::Workspace::FOutlinerItemDetailsId(FUAFLayerStackWorkspaceOutlinerData::StaticStruct()->GetFName()), MakeShared<UE::UAF::LayeringEditor::ULayerStackItemDetails>());
+	
+	WorkspaceEditorModule.RegisterViewportControllerFactory(UUAFLayerStack::StaticClass(), []()
+		{
+			return MakeUnique<FLayerStackViewportController>();
+		});
 }
 
 void FUAFLayeringEditorModule::ShutdownModule()
@@ -53,6 +59,7 @@ void FUAFLayeringEditorModule::ShutdownModule()
 		{
 			WorkspaceEditorModule->UnregisterObjectDocumentType(UUAFLayerStack::StaticClass()->GetClassPathName());
 			WorkspaceEditorModule->UnregisterWorkspaceItemDetails(UE::Workspace::FOutlinerItemDetailsId(FUAFLayerStackWorkspaceOutlinerData::StaticStruct()->GetFName()));
+			WorkspaceEditorModule->UnregisterViewportControllerFactory(UUAFLayerStack::StaticClass());
 		}
 
 		// Unregister layer stack class from UAFEditor
diff --git a/Engine/Plugins/Experimental/UAF/UAFLayering/Source/UAFLayeringEditor/Private/Workspace/LayerStackViewportController.cpp b/Engine/Plugins/Experimental/UAF/UAFLayering/Source/UAFLayeringEditor/Private/Workspace/LayerStackViewportController.cpp
new file mode 100644
--- /dev/null
+++ b/Engine/Plugins/Experimental/UAF/UAFLayering/Source/UAFLayeringEditor/Private/Workspace/LayerStackViewportController.cpp
@@ -0,0 +1,132 @@
+// Copyright Epic Games, Inc. All Rights Reserved.
+
+#include "Workspace/LayerStackViewportController.h"
+
+#include "AdvancedPreviewScene.h"
+#include "UAFLayerStack.h"
+#include "Component/AnimNextComponent.h"
+#include "Engine/PreviewMeshCollection.h"
+#include "Module/UAFSystemAssetData.h"
+#include "UAF/UAFAssetData.h"
+#include "UAF/UAFAssetFactory.h"
+#include "UAF/Viewport/ViewportSceneDescription.h"
+
+namespace UE::UAF::LayeringEditor
+{
+	void FLayerStackViewportController::OnEnter(const FViewportContext& InViewportContext)
+	{
+		const TObjectPtr<const UUAFLayerStack> LayerStack = Cast<UUAFLayerStack>(InViewportContext.OutlinerObject);
+		if (!LayerStack)
+		{
+			UE_LOG(LogAnimation, Warning, TEXT("Workspace Preview: OutlinerObject was null or not of the expected type: [%s]"), *GetNameSafe(UUAFLayerStack::StaticClass()));
+			return;
+		}
+		
+		const TObjectPtr<const UUAFSystem> System = LoadObject<UUAFSystem>(GetTransientPackage(), TEXT("/UAFAnimGraph/Internal/S_SingleGraph.S_SingleGraph"));
+		if (!System)
+		{
+			UE_LOG(LogAnimation, Warning, TEXT("Workspace Preview: Failed to load template system asset"));
+			return;
+		}
+		
+		if (const UUAFViewportSceneDescription* UAFViewportSceneDescription = Cast<UUAFViewportSceneDescription>(InViewportContext.SceneDescription))
+		{
+			// Add main mesh to the scene 
+			if (UAFViewportSceneDescription->SkeletalMesh)
+			{
+				AddMeshToPreview(InViewportContext.PreviewScene, System, LayerStack, UAFViewportSceneDescription->SkeletalMesh);
+			}
+
+			// Add any additional meshes to the scene 
+			if (UAFViewportSceneDescription->AdditionalMeshes)
+			{
+				for (int32 MeshIndex = 0; MeshIndex < UAFViewportSceneDescription->AdditionalMeshes->SkeletalMeshes.Num(); ++MeshIndex)
+				{
+					const TSoftObjectPtr<USkeletalMesh> SkelMeshSoftPtr = UAFViewportSceneDescription->AdditionalMeshes->SkeletalMeshes[MeshIndex].SkeletalMesh;
+				
+					if (!SkelMeshSoftPtr.IsNull())
+					{
+						if (USkeletalMesh* SkeletalMesh = SkelMeshSoftPtr.LoadSynchronous())
+						{
+							AddMeshToPreview(InViewportContext.PreviewScene, System, LayerStack, SkeletalMesh);
+						}
+					}
+				}
+			}
+		}
+	}
+	
+	void FLayerStackViewportController::AddMeshToPreview(FAdvancedPreviewScene* PreviewScene, const TObjectPtr<const UUAFSystem> System, const TObjectPtr<const UUAFLayerStack> LayerStack, USkeletalMesh* InSkeletalMesh)
+	{
+		if (!PreviewScene || !InSkeletalMesh || !System || !LayerStack)
+		{
+			return;
+		}
+		
+		if (const TObjectPtr<UWorld> PreviewWorld = PreviewScene->GetWorld())
+		{
+			if (AActor* PreviewActor = PreviewWorld->SpawnActor<AActor>(AActor::StaticClass()))
+			{
+				PreviewActors.Add(PreviewActor);
+				
+				// Create SkeletalMeshComponent for actor
+				if (USkeletalMeshComponent* SkeletalMeshComponent = NewObject<USkeletalMeshComponent>(PreviewActor))
+				{
+					PreviewScene->AddComponent(SkeletalMeshComponent, FTransform::Identity);
+					
+					PreviewActor->SetRootComponent(SkeletalMeshComponent);
+					SkeletalMeshComponent->SetEnableAnimation(false);
+					SkeletalMeshComponent->SetSkeletalMesh(InSkeletalMesh);
+				}
+				else
+				{
+					UE_LOG(LogAnimation, Warning, TEXT("Workspace Preview: Failed to create a valid SkeletalMeshComponent"));
+				}
+
+				// Create UAFComponent for actor
+				if (UUAFComponent* UAFComponent = NewObject<UUAFComponent>(PreviewActor))
+				{
+					PreviewActor->AddInstanceComponent(UAFComponent);
+					
+					TInstancedStruct<FUAFSystemFactoryAsset> AssetDataFromFactory = UE::UAF::FAssetDataFactory::CreateUAFAssetDataFromObject<FUAFSystemFactoryAsset>(System);
+					UAFComponent->SetAsset(MoveTemp(AssetDataFromFactory));
+					UAFComponent->RegisterComponent();
+					UAFComponent->InitializeComponent();
+					
+					const bool bSet = UAFComponent->SetVariable<const TObjectPtr<const UObject>>(FAnimNextVariableReference::FromName("Graph", System), LayerStack);
+					if (!bSet)
+					{
+						UE_LOG(LogAnimation, Warning, TEXT("Workspace Preview: Could not set %s as the graph asset in %s"), *GetNameSafe(LayerStack), *GetNameSafe(System));
+					}
+				}
+				else
+				{
+					UE_LOG(LogAnimation, Warning, TEXT("Workspace Preview: Failed to create a valid UAFComponent"));
+				}
+			}
+			else
+			{
+				UE_LOG(LogAnimation, Warning, TEXT("Workspace Preview: Failed to spawn preview actor"));
+			}
+		}
+	}
+	
+	void FLayerStackViewportController::OnExit(FAdvancedPreviewScene* PreviewScene)
+	{
+		for (AActor* PreviewActor : PreviewActors)
+		{
+			if (PreviewActor)
+			{
+				if (PreviewScene && PreviewActor->GetRootComponent())
+				{
+					PreviewScene->RemoveComponent(PreviewActor->GetRootComponent());
+				}
+			
+				PreviewActor->Destroy();
+				PreviewActor = nullptr;
+			}
+		}
+		
+		PreviewActors.Empty();
+	}
+}
diff --git a/Engine/Plugins/Experimental/UAF/UAFLayering/Source/UAFLayeringEditor/Private/Workspace/LayerStackViewportController.h b/Engine/Plugins/Experimental/UAF/UAFLayering/Source/UAFLayeringEditor/Private/Workspace/LayerStackViewportController.h
new file mode 100644
--- /dev/null
+++ b/Engine/Plugins/Experimental/UAF/UAFLayering/Source/UAFLayeringEditor/Private/Workspace/LayerStackViewportController.h
@@ -0,0 +1,23 @@
+// Copyright Epic Games, Inc. All Rights Reserved.
+
+#pragma once
+
+#include "UAFLayerStack.h"
+#include "WorkspaceViewportController.h"
+#include "Engine/SkeletalMesh.h"
+#include "Module/AnimNextModule.h"
+
+namespace UE::UAF::LayeringEditor
+{
+	class FLayerStackViewportController : public Workspace::IWorkspaceViewportController
+	{
+	public:
+		virtual void OnEnter(const FViewportContext& InViewportContext) override;
+		virtual void OnExit(FAdvancedPreviewScene* PreviewScene) override;
+
+	private:
+		void AddMeshToPreview(FAdvancedPreviewScene* PreviewScene, const TObjectPtr<const UUAFSystem> System, const TObjectPtr<const UUAFLayerStack> LayerStack, USkeletalMesh* InSkeletalMesh);
+		
+		TArray<AActor*> PreviewActors;
+	};
+}
