UAF Layering: Skip recompilation during slider interaction in the UI and react to changes within layer providers

#Synced-CL 50903586

diff --git a/Engine/Plugins/Experimental/UAF/UAFLayering/Source/UAFLayeringUncookedOnly/Private/Layers/UAFLayer.cpp b/Engine/Plugins/Experimental/UAF/UAFLayering/Source/UAFLayeringUncookedOnly/Private/Layers/UAFLayer.cpp
--- a/Engine/Plugins/Experimental/UAF/UAFLayering/Source/UAFLayeringUncookedOnly/Private/Layers/UAFLayer.cpp
+++ b/Engine/Plugins/Experimental/UAF/UAFLayering/Source/UAFLayeringUncookedOnly/Private/Layers/UAFLayer.cpp
@@ -85,27 +85,36 @@ void UUAFLayer::PreEditChange(FProperty* PropertyAboutToChange)
 void UUAFLayer::PostEditChangeChainProperty(FPropertyChangedChainEvent& PropertyChangedEvent)
 {
 	Super::PostEditChangeChainProperty(PropertyChangedEvent);
-	
-	bool bUpdateLayerLayout = false;
-	if (PropertyChangedEvent.MemberProperty)
+
+	if (PropertyChangedEvent.ChangeType != EPropertyChangeType::Interactive)
 	{
-		if (PropertyChangedEvent.MemberProperty->GetName() == GET_MEMBER_NAME_CHECKED(UUAFLayer, LayerName))
+		bool bUpdateLayerLayout = false;
+		if (PropertyChangedEvent.MemberProperty)
 		{
-			if (!IsLayerNameValid(LayerName) && !PreRenameLayerName.IsNone())
+			if (PropertyChangedEvent.MemberProperty->GetName() == GET_MEMBER_NAME_CHECKED(UUAFLayer, LayerName))
 			{
-				UE_LOGFMT(LogAnimation, Error, "Layer Name {LayerName} is invalid, will reset to previous name {PreLayerName}", LayerName, PreRenameLayerName);
-				LayerName = PreRenameLayerName;
+				if (!IsLayerNameValid(LayerName) && !PreRenameLayerName.IsNone())
+				{
+					UE_LOGFMT(LogAnimation, Error, "Layer Name {LayerName} is invalid, will reset to previous name {PreLayerName}", LayerName, PreRenameLayerName);
+					LayerName = PreRenameLayerName;
+				}
 			}
 		}
-
-		if (PropertyChangedEvent.MemberProperty->GetName() == GET_MEMBER_NAME_CHECKED(UUAFLayer, LayerBlendProvider) 
-			|| PropertyChangedEvent.MemberProperty->GetName() == GET_MEMBER_NAME_CHECKED(UUAFLayer, LayerContentProvider))
+	
+		if (PropertyChangedEvent.PropertyChain.GetHead())
 		{
-			bUpdateLayerLayout = true;
-		}
-	}
+			if (const FProperty* const Property = PropertyChangedEvent.PropertyChain.GetHead()->GetValue())
+			{
+				if (Property->GetName() == GET_MEMBER_NAME_CHECKED(UUAFLayer, LayerBlendProvider) 
+				|| Property->GetName() == GET_MEMBER_NAME_CHECKED(UUAFLayer, LayerContentProvider))
+				{
+					bUpdateLayerLayout = true;
+				}
+			}
+		}	
 	
-	BroadcastModified(EAnimNextEditorDataNotifType::PropertyChanged, bUpdateLayerLayout);
+		BroadcastModified(EAnimNextEditorDataNotifType::PropertyChanged, bUpdateLayerLayout);
+	}
 }
 #endif
 
