UAF Layering: Add workspace outliner support for layers and their contained sub assets

#Synced-CL 51290764

diff --git a/Engine/Plugins/Experimental/UAF/UAFLayering/Source/UAFLayeringEditor/Private/UAFLayeringEditorModule.cpp b/Engine/Plugins/Experimental/UAF/UAFLayering/Source/UAFLayeringEditor/Private/UAFLayeringEditorModule.cpp
--- a/Engine/Plugins/Experimental/UAF/UAFLayering/Source/UAFLayeringEditor/Private/UAFLayeringEditorModule.cpp
+++ b/Engine/Plugins/Experimental/UAF/UAFLayering/Source/UAFLayeringEditor/Private/UAFLayeringEditorModule.cpp
@@ -43,6 +43,7 @@ void FUAFLayeringEditorModule::StartupModule()
 	UE::Workspace::IWorkspaceEditorModule& WorkspaceEditorModule = FModuleManager::Get().LoadModuleChecked<UE::Workspace::IWorkspaceEditorModule>("WorkspaceEditor");
 	WorkspaceEditorModule.RegisterObjectDocumentType(UUAFLayerStack::StaticClass()->GetClassPathName(), LayerStackDocumentArgs);
 	WorkspaceEditorModule.RegisterWorkspaceItemDetails(UE::Workspace::FOutlinerItemDetailsId(FUAFLayerStackWorkspaceOutlinerData::StaticStruct()->GetFName()), MakeShared<UE::UAF::LayeringEditor::ULayerStackItemDetails>());
+	WorkspaceEditorModule.RegisterWorkspaceItemDetails(UE::Workspace::FOutlinerItemDetailsId(FUAFLayerStackLayerOutlinerData::StaticStruct()->GetFName()), MakeShared<UE::UAF::LayeringEditor::ULayerStackLayerItemDetails>());
 	
 	WorkspaceEditorModule.RegisterViewportControllerFactory(UUAFLayerStack::StaticClass(), []()
 		{
@@ -59,6 +60,7 @@ void FUAFLayeringEditorModule::ShutdownModule()
 		{
 			WorkspaceEditorModule->UnregisterObjectDocumentType(UUAFLayerStack::StaticClass()->GetClassPathName());
 			WorkspaceEditorModule->UnregisterWorkspaceItemDetails(UE::Workspace::FOutlinerItemDetailsId(FUAFLayerStackWorkspaceOutlinerData::StaticStruct()->GetFName()));
+			WorkspaceEditorModule->UnregisterWorkspaceItemDetails(UE::Workspace::FOutlinerItemDetailsId(FUAFLayerStackLayerOutlinerData::StaticStruct()->GetFName()));
 			WorkspaceEditorModule->UnregisterViewportControllerFactory(UUAFLayerStack::StaticClass());
 		}
 
diff --git a/Engine/Plugins/Experimental/UAF/UAFLayering/Source/UAFLayeringEditor/Private/Workspace/LayerStackItemDetails.cpp b/Engine/Plugins/Experimental/UAF/UAFLayering/Source/UAFLayeringEditor/Private/Workspace/LayerStackItemDetails.cpp
--- a/Engine/Plugins/Experimental/UAF/UAFLayering/Source/UAFLayeringEditor/Private/Workspace/LayerStackItemDetails.cpp
+++ b/Engine/Plugins/Experimental/UAF/UAFLayering/Source/UAFLayeringEditor/Private/Workspace/LayerStackItemDetails.cpp
@@ -2,14 +2,129 @@
 
 #include "LayerStackItemDetails.h"
 
+#include "IWorkspaceEditor.h"
+#include "IWorkspaceEditorModule.h"
+#include "ToolMenuContext.h"
 #include "UAFLayeringStyle.h"
+#include "UAFLayerStack.h"
 #include "UAFStyle.h"
+#include "WorkspaceItemMenuContext.h"
+#include "Layers/UAFLayer.h"
+#include "Toolkits/AssetEditorToolkitMenuContext.h"
+#include "Workspace/UAFLayerStackWorkspaceOutlinerData.h"
 
 namespace UE::UAF::LayeringEditor
 {
-const FSlateBrush* ULayerStackItemDetails::GetItemIcon(const FWorkspaceOutlinerItemExport& Export) const
-{
-	return FUAFLayeringStyle::Get().GetBrush(TEXT("ClassIcon.UAFLayerStack"));
-}
+	bool ULayerStackItemDetails::HandleSelected(const FToolMenuContext& ToolMenuContext) const
+	{
+		const UWorkspaceItemMenuContext* WorkspaceItemContext = ToolMenuContext.FindContext<UWorkspaceItemMenuContext>();
+		const UAssetEditorToolkitMenuContext* AssetEditorContext = ToolMenuContext.FindContext<UAssetEditorToolkitMenuContext>(); 
+		if (WorkspaceItemContext && AssetEditorContext)
+		{
+			if(const TSharedPtr<UE::Workspace::IWorkspaceEditor> WorkspaceEditor = StaticCastSharedPtr<UE::Workspace::IWorkspaceEditor>(AssetEditorContext->Toolkit.Pin()))
+			{
+				if (WorkspaceItemContext->SelectedExports.Num() > 0)
+				{
+					const FWorkspaceOutlinerItemExport& SelectedExport = WorkspaceItemContext->SelectedExports[0].GetResolvedExport();
+					if (UUAFLayerStack* LayerStack = SelectedExport.GetFirstObjectOfType<UUAFLayerStack>())
+					{
+						if (const UUAFLayerStack_EditorData* EditorData = UncookedOnly::FUtils::GetEditorData<UUAFLayerStack_EditorData>(LayerStack))
+						{
+							EditorData->ClearSelectedLayer();
+							return true;
+						}
+					}
+				}
+			}
+		}
+		
+		return false;
+	}
+
+	const FSlateBrush* ULayerStackItemDetails::GetItemIcon(const FWorkspaceOutlinerItemExport& Export) const
+	{
+		return FUAFLayeringStyle::Get().GetBrush(TEXT("ClassIcon.UAFLayerStack"));
+	}
+
+	bool ULayerStackLayerItemDetails::IsExpandedByDefault() const
+	{
+		return false;
+	}
+
+	FString ULayerStackLayerItemDetails::GetDisplayString(const FWorkspaceOutlinerItemExport& Export) const
+	{
+		const TInstancedStruct<FWorkspaceOutlinerItemData>& Data = Export.GetData();
+		if (Data.IsValid() && Data.GetScriptStruct() == FUAFLayerStackLayerOutlinerData::StaticStruct())
+		{
+			const FUAFLayerStackLayerOutlinerData& LayerData = Data.Get<FUAFLayerStackLayerOutlinerData>();
+			return LayerData.LayerName.ToString();
+		}
+		
+		return IWorkspaceOutlinerItemDetails::GetDisplayString(Export);
+	}
+
+	const FSlateBrush* ULayerStackLayerItemDetails::GetItemIcon(const FWorkspaceOutlinerItemExport& Export) const
+	{
+		return FUAFLayeringStyle::Get().GetBrush(TEXT("ClassIcon.UAFLayer"));
+	}
 	
+	bool ULayerStackLayerItemDetails::HandleSelected(const FToolMenuContext& ToolMenuContext) const
+	{
+		const UWorkspaceItemMenuContext* WorkspaceItemContext = ToolMenuContext.FindContext<UWorkspaceItemMenuContext>();
+		const UAssetEditorToolkitMenuContext* AssetEditorContext = ToolMenuContext.FindContext<UAssetEditorToolkitMenuContext>(); 
+		if (WorkspaceItemContext && AssetEditorContext)
+		{
+			if(const TSharedPtr<UE::Workspace::IWorkspaceEditor> WorkspaceEditor = StaticCastSharedPtr<UE::Workspace::IWorkspaceEditor>(AssetEditorContext->Toolkit.Pin()))
+			{
+				if (WorkspaceItemContext->SelectedExports.Num() > 0)
+				{
+					const FWorkspaceOutlinerItemExport& SelectedExport = WorkspaceItemContext->SelectedExports[0].GetResolvedExport();
+					if (UUAFLayer* SelectedLayer = SelectedExport.GetFirstObjectOfType<UUAFLayer>())
+					{
+						if (const UUAFLayerStack_EditorData* EditorData = SelectedLayer->GetTypedOuter<UUAFLayerStack_EditorData>())
+						{
+							EditorData->SelectLayer(SelectedLayer);
+							return true;
+						}
+					}
+				}
+			}
+		}
+		
+		return false;
+	}
+
+	bool ULayerStackLayerItemDetails::HandleDoubleClick(const FToolMenuContext& ToolMenuContext) const
+	{
+		const UWorkspaceItemMenuContext* WorkspaceItemContext = ToolMenuContext.FindContext<UWorkspaceItemMenuContext>();
+		const UAssetEditorToolkitMenuContext* AssetEditorContext = ToolMenuContext.FindContext<UAssetEditorToolkitMenuContext>(); 
+		if (WorkspaceItemContext && AssetEditorContext)
+		{
+			if(const TSharedPtr<UE::Workspace::IWorkspaceEditor> WorkspaceEditor = StaticCastSharedPtr<UE::Workspace::IWorkspaceEditor>(AssetEditorContext->Toolkit.Pin()))
+			{
+				if (WorkspaceItemContext->SelectedExports.Num() > 0)
+				{
+					const FWorkspaceOutlinerItemExport& SelectedExport = WorkspaceItemContext->SelectedExports[0].GetResolvedExport();
+					if (UUAFLayer* SelectedLayer = SelectedExport.GetFirstObjectOfType<UUAFLayer>())
+					{
+						if (UUAFLayerStack_EditorData* EditorData = SelectedLayer->GetTypedOuter<UUAFLayerStack_EditorData>())
+						{
+							if (const UUAFRigVMAsset* RigVMAsset = UncookedOnly::FUtils::GetAsset(EditorData))
+							{
+								const Workspace::FWorkspaceDocument& Document = WorkspaceEditor->GetFocusedWorkspaceDocument();	
+								const FWorkspaceOutlinerItemExport RelativeAssetExport = Document.Export.GetRelativeAssetExport(RigVMAsset);
+								const FWorkspaceOutlinerItemExport AssetExport = UncookedOnly::FUtils::MakeAssetReferenceExport(RigVMAsset, RelativeAssetExport);
+
+								WorkspaceEditor->OpenExports({AssetExport});
+								EditorData->SelectLayer(SelectedLayer);
+								return true;
+							}
+						}
+					}
+				}
+			}
+		}
+		
+		return false;
+	}
 }
diff --git a/Engine/Plugins/Experimental/UAF/UAFLayering/Source/UAFLayeringEditor/Private/Workspace/LayerStackItemDetails.h b/Engine/Plugins/Experimental/UAF/UAFLayering/Source/UAFLayeringEditor/Private/Workspace/LayerStackItemDetails.h
--- a/Engine/Plugins/Experimental/UAF/UAFLayering/Source/UAFLayeringEditor/Private/Workspace/LayerStackItemDetails.h
+++ b/Engine/Plugins/Experimental/UAF/UAFLayering/Source/UAFLayeringEditor/Private/Workspace/LayerStackItemDetails.h
@@ -6,10 +6,21 @@
 
 namespace UE::UAF::LayeringEditor
 {
-class ULayerStackItemDetails : public UE::Workspace::IWorkspaceOutlinerItemDetails
-{
-public:
-	virtual const FSlateBrush* GetItemIcon(const FWorkspaceOutlinerItemExport& Export) const override;
-};
+	class ULayerStackItemDetails : public UE::Workspace::IWorkspaceOutlinerItemDetails
+	{
+	public:
+		virtual bool HandleSelected(const FToolMenuContext& ToolMenuContext) const override;
+		virtual const FSlateBrush* GetItemIcon(const FWorkspaceOutlinerItemExport& Export) const override;
+	};
+	
+	class ULayerStackLayerItemDetails : public UE::Workspace::IWorkspaceOutlinerItemDetails
+	{
+	public:
+		virtual bool IsExpandedByDefault() const override;
+		virtual FString GetDisplayString(const FWorkspaceOutlinerItemExport& Export) const override;
+		virtual const FSlateBrush* GetItemIcon(const FWorkspaceOutlinerItemExport& Export) const override;
+		virtual bool HandleSelected(const FToolMenuContext& ToolMenuContext) const override;
+		virtual bool HandleDoubleClick(const FToolMenuContext& ToolMenuContext) const override;
+	};
 }
 
diff --git a/Engine/Plugins/Experimental/UAF/UAFLayering/Source/UAFLayeringUncookedOnly/Internal/Layers/UAFLayer.h b/Engine/Plugins/Experimental/UAF/UAFLayering/Source/UAFLayeringUncookedOnly/Internal/Layers/UAFLayer.h
--- a/Engine/Plugins/Experimental/UAF/UAFLayering/Source/UAFLayeringUncookedOnly/Internal/Layers/UAFLayer.h
+++ b/Engine/Plugins/Experimental/UAF/UAFLayering/Source/UAFLayeringUncookedOnly/Internal/Layers/UAFLayer.h
@@ -60,6 +60,7 @@ public:
 #if WITH_EDITOR
 	virtual void PreEditChange(FProperty* PropertyAboutToChange) override;
 	virtual void PostEditChangeChainProperty(struct FPropertyChangedChainEvent& PropertyChangedEvent) override;
+	virtual void GetObjectReferences(TArray<const UObject*>& OutReferencedObjects) const;
 #endif
 	
 protected:
diff --git a/Engine/Plugins/Experimental/UAF/UAFLayering/Source/UAFLayeringUncookedOnly/Internal/Layers/UAFLayerAssetProvider.h b/Engine/Plugins/Experimental/UAF/UAFLayering/Source/UAFLayeringUncookedOnly/Internal/Layers/UAFLayerAssetProvider.h
--- a/Engine/Plugins/Experimental/UAF/UAFLayering/Source/UAFLayeringUncookedOnly/Internal/Layers/UAFLayerAssetProvider.h
+++ b/Engine/Plugins/Experimental/UAF/UAFLayering/Source/UAFLayeringUncookedOnly/Internal/Layers/UAFLayerAssetProvider.h
@@ -35,6 +35,10 @@ public:
 
 	void SetLayerAsset(const UObject* InNewAsset);
 	
+#if WITH_EDITOR
+	virtual void GetObjectReferences(TArray<const UObject*>& OutReferencedObjects) const override;
+#endif
+	
 #if WITH_EDITORONLY_DATA
 	void PostSerialize(const FArchive& Ar);
 #endif
diff --git a/Engine/Plugins/Experimental/UAF/UAFLayering/Source/UAFLayeringUncookedOnly/Internal/Layers/UAFLayerBlendProviderBase.h b/Engine/Plugins/Experimental/UAF/UAFLayering/Source/UAFLayeringUncookedOnly/Internal/Layers/UAFLayerBlendProviderBase.h
--- a/Engine/Plugins/Experimental/UAF/UAFLayering/Source/UAFLayeringUncookedOnly/Internal/Layers/UAFLayerBlendProviderBase.h
+++ b/Engine/Plugins/Experimental/UAF/UAFLayering/Source/UAFLayeringUncookedOnly/Internal/Layers/UAFLayerBlendProviderBase.h
@@ -36,6 +36,11 @@ public:
 	
 	void SetAlwaysUpdateChildren(bool bInAlwaysUpdate) { bAlwaysUpdateChildren = bInAlwaysUpdate; };
 	bool GetAlwaysUpdateChildren() const { return bAlwaysUpdateChildren; };
+	
+#if WITH_EDITOR
+	virtual void GetObjectReferences(TArray<const UObject*>& OutReferencedObjects) const {};
+#endif
+	
 private:
 	bool bAlwaysUpdateChildren = false;	
 };
diff --git a/Engine/Plugins/Experimental/UAF/UAFLayering/Source/UAFLayeringUncookedOnly/Internal/Layers/UAFLayerContentProviderBase.h b/Engine/Plugins/Experimental/UAF/UAFLayering/Source/UAFLayeringUncookedOnly/Internal/Layers/UAFLayerContentProviderBase.h
--- a/Engine/Plugins/Experimental/UAF/UAFLayering/Source/UAFLayeringUncookedOnly/Internal/Layers/UAFLayerContentProviderBase.h
+++ b/Engine/Plugins/Experimental/UAF/UAFLayering/Source/UAFLayeringUncookedOnly/Internal/Layers/UAFLayerContentProviderBase.h
@@ -21,9 +21,14 @@ public:
 
 	// Responsible for creating the traits to play the content of this layer 
 	// @return The output pin of this layers content trait stack which can be linked to the blend trait stack of this layer 
-	virtual URigVMPin* CreateLayerContentTrait(UE::UAF::Layering::FLayerCreationContext& LayerCreationContext)  {return nullptr;} ;
+	virtual URigVMPin* CreateLayerContentTrait(UE::UAF::Layering::FLayerCreationContext& LayerCreationContext)  { return nullptr; };
 	
 	// Responsible for creating the widget to represent this layers content in the layer stack editor 
-	virtual TSharedRef<SWidget> CreateLayerContentWidget(UUAFLayer* InLayer) {return SNullWidget::NullWidget;} ;  
+	virtual TSharedRef<SWidget> CreateLayerContentWidget(UUAFLayer* InLayer) { return SNullWidget::NullWidget; };  
+	
+#if WITH_EDITOR
+	virtual void GetObjectReferences(TArray<const UObject*>& OutReferencedObjects) const {};
+#endif
+	
 };
 
diff --git a/Engine/Plugins/Experimental/UAF/UAFLayering/Source/UAFLayeringUncookedOnly/Internal/UAFLayerStack_EditorData.h b/Engine/Plugins/Experimental/UAF/UAFLayering/Source/UAFLayeringUncookedOnly/Internal/UAFLayerStack_EditorData.h
--- a/Engine/Plugins/Experimental/UAF/UAFLayering/Source/UAFLayeringUncookedOnly/Internal/UAFLayerStack_EditorData.h
+++ b/Engine/Plugins/Experimental/UAF/UAFLayering/Source/UAFLayeringUncookedOnly/Internal/UAFLayerStack_EditorData.h
@@ -14,6 +14,7 @@ class UUAFLayerStack;
 #define UE_API UAFLAYERINGUNCOOKEDONLY_API
 
 DECLARE_MULTICAST_DELEGATE(FLayerLayoutChangedDelegate);
+DECLARE_MULTICAST_DELEGATE_OneParam(FLayerSelectionChangedDelegate, const TWeakObjectPtr<UUAFLayer>);
 
 UCLASS()
 class UUAFLayerStack_EditorData : public UUAFAnimGraph_EditorData, public FEditorUndoClient
@@ -41,6 +42,9 @@ public:
 	void DeleteLayer(UUAFLayer* LayerToDelete);
 	void SetLayerState(UUAFLayer* Layer, EUAFLayerState LayerState);
 	
+	UE_API void SelectLayer(UUAFLayer* Layer) const;
+	UE_API void ClearSelectedLayer() const;
+	
 	UUAFLayer* CreateDefaultAssetBasedLayer(const FAssetData& InAsset);
 	
 	void BroadcastLayerStackChanged(EAnimNextEditorDataNotifType InType, bool bRefreshUI = false);
@@ -63,6 +67,7 @@ public:
 	TWeakObjectPtr<URigVMGraph> CreatedGraph = nullptr;
 	
 	FLayerLayoutChangedDelegate OnLayerLayoutChanged;
+	FLayerSelectionChangedDelegate OnLayerSelectionChanged;
 
 private:
 	UPROPERTY(Instanced, VisibleAnywhere, Category = "LayerStack")
diff --git a/Engine/Plugins/Experimental/UAF/UAFLayering/Source/UAFLayeringUncookedOnly/Internal/Widgets/SLayerStack.h b/Engine/Plugins/Experimental/UAF/UAFLayering/Source/UAFLayeringUncookedOnly/Internal/Widgets/SLayerStack.h
--- a/Engine/Plugins/Experimental/UAF/UAFLayering/Source/UAFLayeringUncookedOnly/Internal/Widgets/SLayerStack.h
+++ b/Engine/Plugins/Experimental/UAF/UAFLayering/Source/UAFLayeringUncookedOnly/Internal/Widgets/SLayerStack.h
@@ -35,6 +35,7 @@ private:
 	void OnListSelectionChanged(TWeakObjectPtr<UUAFLayer> SelectedItem, ESelectInfo::Type SelectInfo) const;
 	TSharedPtr<SWidget> OnConstructContextMenu() const;
 	void OnLayerLayoutChanged();
+	void OnLayerSelectionChanged(const TWeakObjectPtr<UUAFLayer> SelectedLayer) const;
 	void CacheLayersForList();
 private:
 	TSharedPtr<SListView<TWeakObjectPtr<UUAFLayer>>> LayerListView = nullptr;
diff --git a/Engine/Plugins/Experimental/UAF/UAFLayering/Source/UAFLayeringUncookedOnly/Internal/Workspace/UAFLayerStackWorkspaceOutlinerData.h b/Engine/Plugins/Experimental/UAF/UAFLayering/Source/UAFLayeringUncookedOnly/Internal/Workspace/UAFLayerStackWorkspaceOutlinerData.h
--- a/Engine/Plugins/Experimental/UAF/UAFLayering/Source/UAFLayeringUncookedOnly/Internal/Workspace/UAFLayerStackWorkspaceOutlinerData.h
+++ b/Engine/Plugins/Experimental/UAF/UAFLayering/Source/UAFLayeringUncookedOnly/Internal/Workspace/UAFLayerStackWorkspaceOutlinerData.h
@@ -12,4 +12,15 @@ struct FUAFLayerStackWorkspaceOutlinerData : public FAnimNextRigVMAssetOutlinerD
 	GENERATED_BODY()
 	
 	FUAFLayerStackWorkspaceOutlinerData() = default;
-};
\ No newline at end of file
+};
+
+USTRUCT()
+struct FUAFLayerStackLayerOutlinerData : public FWorkspaceOutlinerItemData
+{
+	GENERATED_BODY()
+	
+	FUAFLayerStackLayerOutlinerData() = default;
+
+	UPROPERTY()
+	FName LayerName = NAME_None;
+};
diff --git a/Engine/Plugins/Experimental/UAF/UAFLayering/Source/UAFLayeringUncookedOnly/Private/Layers/UAFLayer.cpp b/Engine/Plugins/Experimental/UAF/UAFLayering/Source/UAFLayeringUncookedOnly/Private/Layers/UAFLayer.cpp
--- a/Engine/Plugins/Experimental/UAF/UAFLayering/Source/UAFLayeringUncookedOnly/Private/Layers/UAFLayer.cpp
+++ b/Engine/Plugins/Experimental/UAF/UAFLayering/Source/UAFLayeringUncookedOnly/Private/Layers/UAFLayer.cpp
@@ -116,6 +116,19 @@ void UUAFLayer::PostEditChangeChainProperty(FPropertyChangedChainEvent& Property
 		BroadcastModified(EAnimNextEditorDataNotifType::PropertyChanged, bUpdateLayerLayout);
 	}
 }
+
+void UUAFLayer::GetObjectReferences(TArray<const UObject*>& OutReferencedObjects) const
+{
+	if (IsLayerBlendProviderValid())
+	{
+		LayerBlendProvider.GetPtr()->GetObjectReferences(OutReferencedObjects);
+	}
+
+	if (IsLayerContentProviderValid())
+	{
+		LayerContentProvider.GetPtr()->GetObjectReferences(OutReferencedObjects);
+	}
+}
 #endif
 
 void UUAFLayer::RenameLayer(const FName NewName)
diff --git a/Engine/Plugins/Experimental/UAF/UAFLayering/Source/UAFLayeringUncookedOnly/Private/Layers/UAFLayerAssetProvider.cpp b/Engine/Plugins/Experimental/UAF/UAFLayering/Source/UAFLayeringUncookedOnly/Private/Layers/UAFLayerAssetProvider.cpp
--- a/Engine/Plugins/Experimental/UAF/UAFLayering/Source/UAFLayeringUncookedOnly/Private/Layers/UAFLayerAssetProvider.cpp
+++ b/Engine/Plugins/Experimental/UAF/UAFLayering/Source/UAFLayeringUncookedOnly/Private/Layers/UAFLayerAssetProvider.cpp
@@ -66,35 +66,6 @@ URigVMPin* FUAFLayerAssetProvider::CreateLayerContentTrait(UE::UAF::Layering::FL
 	return BlendTraitStackResult;
 }
 
-void FUAFLayerAssetProvider::SetLayerAsset(const UObject* InNewAsset)
-{
-	const FScopedTransaction Transaction(LOCTEXT("OnLayerAssetSelectionChangedAction", "Set Layer Asset"));
-	if (OuterLayer.IsValid())
-	{
-		OuterLayer->Modify();
-	}
-
-	if (InNewAsset)
-	{
-		AssetData = UE::UAF::FAssetDataFactory::CreateUAFAssetDataFromObject<FUAFGraphFactoryAsset>(InNewAsset);
-	}
-	else
-	{
-		AssetData.Reset();
-	}
-	
-	if (OuterLayer.IsValid())
-	{
-		OuterLayer->BroadcastModified(EAnimNextEditorDataNotifType::PropertyChanged);
-	}
-}
-
-void FUAFLayerAssetProvider::OnLayerAssetSelectionChanged(const FAssetData& InAssetData)
-{
-	UObject* Obj = Cast<UObject>(InAssetData.GetAsset());
-	SetLayerAsset(Obj);
-}
-
 TSharedRef<SWidget> FUAFLayerAssetProvider::CreateLayerContentWidget(UUAFLayer* InLayer)
 {
 	OuterLayer = InLayer;
@@ -129,6 +100,45 @@ TSharedRef<SWidget> FUAFLayerAssetProvider::CreateLayerContentWidget(UUAFLayer*
 		];
 }
 
+void FUAFLayerAssetProvider::SetLayerAsset(const UObject* InNewAsset)
+{
+	const FScopedTransaction Transaction(LOCTEXT("OnLayerAssetSelectionChangedAction", "Set Layer Asset"));
+	if (OuterLayer.IsValid())
+	{
+		OuterLayer->Modify();
+	}
+
+	if (InNewAsset)
+	{
+		AssetData = UE::UAF::FAssetDataFactory::CreateUAFAssetDataFromObject<FUAFGraphFactoryAsset>(InNewAsset);
+	}
+	else
+	{
+		AssetData.Reset();
+	}
+	
+	if (OuterLayer.IsValid())
+	{
+		OuterLayer->BroadcastModified(EAnimNextEditorDataNotifType::PropertyChanged);
+	}
+}
+
+void FUAFLayerAssetProvider::OnLayerAssetSelectionChanged(const FAssetData& InAssetData)
+{
+	UObject* Obj = Cast<UObject>(InAssetData.GetAsset());
+	SetLayerAsset(Obj);
+}
+
+#if WITH_EDITOR
+void FUAFLayerAssetProvider::GetObjectReferences(TArray<const UObject*>& OutReferencedObjects) const
+{
+	if (AssetData.IsValid())
+	{
+		AssetData.GetPtr()->GetObjectReferences(OutReferencedObjects);
+	}
+}
+#endif
+
 #if WITH_EDITORONLY_DATA
 void FUAFLayerAssetProvider::PostSerialize(const FArchive& Ar)
 {
diff --git a/Engine/Plugins/Experimental/UAF/UAFLayering/Source/UAFLayeringUncookedOnly/Private/UAFLayerStack_EditorData.cpp b/Engine/Plugins/Experimental/UAF/UAFLayering/Source/UAFLayeringUncookedOnly/Private/UAFLayerStack_EditorData.cpp
--- a/Engine/Plugins/Experimental/UAF/UAFLayering/Source/UAFLayeringUncookedOnly/Private/UAFLayerStack_EditorData.cpp
+++ b/Engine/Plugins/Experimental/UAF/UAFLayering/Source/UAFLayeringUncookedOnly/Private/UAFLayerStack_EditorData.cpp
@@ -226,6 +226,16 @@ void UUAFLayerStack_EditorData::SetLayerState(UUAFLayer* Layer, EUAFLayerState L
 	BroadcastModified(EAnimNextEditorDataNotifType::PropertyChanged, this);
 }
 
+void UUAFLayerStack_EditorData::SelectLayer(UUAFLayer* Layer) const
+{
+	OnLayerSelectionChanged.Broadcast(Layer);
+}
+
+void UUAFLayerStack_EditorData::ClearSelectedLayer() const
+{
+	OnLayerSelectionChanged.Broadcast(nullptr);
+}
+
 int32 UUAFLayerStack_EditorData::GetIndexForLayer(UUAFLayer* Layer) const 
 {
 	return GetAllLayers().Find(Layer);
diff --git a/Engine/Plugins/Experimental/UAF/UAFLayering/Source/UAFLayeringUncookedOnly/Private/UAFLayeringStyle.cpp b/Engine/Plugins/Experimental/UAF/UAFLayering/Source/UAFLayeringUncookedOnly/Private/UAFLayeringStyle.cpp
--- a/Engine/Plugins/Experimental/UAF/UAFLayering/Source/UAFLayeringUncookedOnly/Private/UAFLayeringStyle.cpp
+++ b/Engine/Plugins/Experimental/UAF/UAFLayering/Source/UAFLayeringUncookedOnly/Private/UAFLayeringStyle.cpp
@@ -68,6 +68,7 @@ FUAFLayeringStyle::FUAFLayeringStyle()
 		);
 	
 	Set("ClassIcon.UAFLayerStack", new IMAGE_BRUSH("Icons/layers_16x", Icon16x16));
+	Set("ClassIcon.UAFLayer", new IMAGE_BRUSH("Icons/layer_16x", Icon16x16));
 	
 	Set("Layer.VisibleLayerIcon", new CORE_IMAGE_BRUSH_SVG("Starship/Common/visible", Icon16x16));
 	Set("Layer.DisabledLayerIcon", new IMAGE_BRUSH("Icons/SourceControlOff_16x", Icon16x16));
diff --git a/Engine/Plugins/Experimental/UAF/UAFLayering/Source/UAFLayeringUncookedOnly/Private/Widgets/SLayerStack.cpp b/Engine/Plugins/Experimental/UAF/UAFLayering/Source/UAFLayeringUncookedOnly/Private/Widgets/SLayerStack.cpp
--- a/Engine/Plugins/Experimental/UAF/UAFLayering/Source/UAFLayeringUncookedOnly/Private/Widgets/SLayerStack.cpp
+++ b/Engine/Plugins/Experimental/UAF/UAFLayering/Source/UAFLayeringUncookedOnly/Private/Widgets/SLayerStack.cpp
@@ -32,6 +32,8 @@ void SLayerStack::Construct(const FArguments& InArgs)
 
 	UUAFLayerStack_EditorData* EditorData = UE::UAF::UncookedOnly::FUtils::GetEditorData<UUAFLayerStack_EditorData>(LayerStack.Get());
 	EditorData->OnLayerLayoutChanged.AddSP(this, &SLayerStack::OnLayerLayoutChanged);
+	EditorData->OnLayerSelectionChanged.AddSP(this, &SLayerStack::OnLayerSelectionChanged);
+	
 	CacheLayersForList();
 
 	// make commands for context menu
@@ -412,7 +414,19 @@ void SLayerStack::OnLayerLayoutChanged()
 	CacheLayersForList();
 	LayerListView->RebuildList();
 }
-	
+
+void SLayerStack::OnLayerSelectionChanged(const TWeakObjectPtr<UUAFLayer> SelectedLayer) const
+{
+	if (SelectedLayer.IsValid())
+	{
+		LayerListView->SetSelection(SelectedLayer);
+	}
+	else
+	{
+		LayerListView->ClearSelection();
+	}
+}
+
 void SLayerStack::CacheLayersForList()
 {
 	if (!LayerStack.IsValid())
diff --git a/Engine/Plugins/Experimental/UAF/UAFLayering/Source/UAFLayeringUncookedOnly/Private/Workspace/UAFLayerStackWorkspaceAssetUserData.cpp b/Engine/Plugins/Experimental/UAF/UAFLayering/Source/UAFLayeringUncookedOnly/Private/Workspace/UAFLayerStackWorkspaceAssetUserData.cpp
--- a/Engine/Plugins/Experimental/UAF/UAFLayering/Source/UAFLayeringUncookedOnly/Private/Workspace/UAFLayerStackWorkspaceAssetUserData.cpp
+++ b/Engine/Plugins/Experimental/UAF/UAFLayering/Source/UAFLayeringUncookedOnly/Private/Workspace/UAFLayerStackWorkspaceAssetUserData.cpp
@@ -2,14 +2,66 @@
 
 
 #include "Workspace/UAFLayerStackWorkspaceAssetUserData.h"
+
 #include "UAFLayerStack.h"
+#include "UAFLayerStack_EditorData.h"
+#include "UncookedOnlyUtils.h"
+#include "Layers/UAFLayer.h"
+#include "Workspace/UAFLayerStackWorkspaceOutlinerData.h"
 
 void UUAFLayerStackWorkspaceAssetUserData::GetRootAssetExport(FAssetRegistryTagsContext Context) const
 {
-	Super::GetRootAssetExport(Context);
+	if (UUAFLayerStack* LayerStack = CastChecked<UUAFLayerStack>(GetOuter()))
+	{
+		const FName LayerStackIdentifier = LayerStack->GetFName();
+		FWorkspaceOutlinerItemExport& RootAssetExport = CachedExports.Exports.Add_GetRef(FWorkspaceOutlinerItemExport(LayerStackIdentifier, LayerStack));
+		RootAssetExport.GetData().InitializeAsScriptStruct(FUAFLayerStackWorkspaceOutlinerData::StaticStruct());
+		RootAssetExport.GetData().GetMutable<FUAFLayerStackWorkspaceOutlinerData>().SoftAssetPtr = LayerStack;
+	}
 }
 
 void UUAFLayerStackWorkspaceAssetUserData::GetWorkspaceAssetExports(FAssetRegistryTagsContext Context) const
 {
 	Super::GetWorkspaceAssetExports(Context);
+	
+	if (UUAFLayerStack* LayerStack = CastChecked<UUAFLayerStack>(GetOuter()))
+	{
+		FWorkspaceOutlinerItemExport RootAssetExport = CachedExports.Exports.Num() > 0 ? CachedExports.Exports[0] : FWorkspaceOutlinerItemExport();
+		
+		// Add each layer as a workspace export 
+		UUAFLayerStack_EditorData* EditorData = UE::UAF::UncookedOnly::FUtils::GetEditorData<UUAFLayerStack_EditorData>(LayerStack);
+		if (EditorData)
+		{
+			const TArray<TObjectPtr<UUAFLayer>> Layers = EditorData->GetAllLayers();
+			for (const TObjectPtr<UUAFLayer>& Layer : Layers)
+			{
+				if (Layer == nullptr)
+				{
+					continue;
+				}
+				
+				// Add layer to the export
+				const FSoftObjectPath LayerSoftPath = Layer;
+				FWorkspaceOutlinerItemExport& LayerExport = CachedExports.Exports.Add_GetRef(FWorkspaceOutlinerItemExport(*LayerSoftPath.GetSubPathUtf8String(), RootAssetExport));
+				LayerExport.GetData().InitializeAsScriptStruct(FUAFLayerStackLayerOutlinerData::StaticStruct());
+				
+				FUAFLayerStackLayerOutlinerData& AssetData = LayerExport.GetData().GetMutable<FUAFLayerStackLayerOutlinerData>();
+				AssetData.LayerName = Layer->GetLayerName();
+				
+				// Add UObjects referenced by layers to the export 
+				TArray<const UObject*> ReferencedObjects;
+				Layer->GetObjectReferences(ReferencedObjects);
+				for (const UObject* Object : ReferencedObjects)
+				{
+					FSoftObjectPath SoftObjectPath = Object;
+					
+					FWorkspaceOutlinerItemExport& AssetReference = CachedExports.Exports.Add_GetRef(FWorkspaceOutlinerItemExport(FName(SoftObjectPath.ToString()), LayerExport));
+					
+					AssetReference.GetData().InitializeAsScriptStruct(FWorkspaceOutlinerAssetReferenceItemData::StaticStruct());
+					FWorkspaceOutlinerAssetReferenceItemData& ReferenceData = AssetReference.GetData().GetMutable<FWorkspaceOutlinerAssetReferenceItemData>();
+					ReferenceData.ReferredObjectPath = SoftObjectPath;
+				}
+			}
+		}
+	}
 }
