UAF Layering: Expand outliner items per default

#Synced-CL 51408588

diff --git a/Engine/Plugins/Experimental/UAF/UAFLayering/Source/UAFLayeringEditor/Private/Workspace/LayerStackItemDetails.cpp b/Engine/Plugins/Experimental/UAF/UAFLayering/Source/UAFLayeringEditor/Private/Workspace/LayerStackItemDetails.cpp
--- a/Engine/Plugins/Experimental/UAF/UAFLayering/Source/UAFLayeringEditor/Private/Workspace/LayerStackItemDetails.cpp
+++ b/Engine/Plugins/Experimental/UAF/UAFLayering/Source/UAFLayeringEditor/Private/Workspace/LayerStackItemDetails.cpp
@@ -48,7 +48,7 @@ namespace UE::UAF::LayeringEditor
 
 	bool ULayerStackLayerItemDetails::IsExpandedByDefault() const
 	{
-		return false;
+		return true;
 	}
 
 	FString ULayerStackLayerItemDetails::GetDisplayString(const FWorkspaceOutlinerItemExport& Export) const
