UAF Layering:
- Add auto naming of new layers
- Minor refactor of utility functions in the editor data

#Synced-CL 51454554

diff --git a/Engine/Plugins/Experimental/UAF/UAFLayering/Source/UAFLayeringUncookedOnly/Internal/Layers/UAFLayer.h b/Engine/Plugins/Experimental/UAF/UAFLayering/Source/UAFLayeringUncookedOnly/Internal/Layers/UAFLayer.h
--- a/Engine/Plugins/Experimental/UAF/UAFLayering/Source/UAFLayeringUncookedOnly/Internal/Layers/UAFLayer.h
+++ b/Engine/Plugins/Experimental/UAF/UAFLayering/Source/UAFLayeringUncookedOnly/Internal/Layers/UAFLayer.h
@@ -70,7 +70,7 @@ protected:
 
 	// The name of this layer. It has to be unique within one layer stack asset. 
 	UPROPERTY(EditAnywhere, Category = "Layer")
-	FName LayerName;
+	FName LayerName = NAME_None;
 	
 	UPROPERTY(EditAnywhere, Category = "Layer", meta = (ExcludeBaseStruct))
 	TInstancedStruct<FUAFLayerContentProviderBase> LayerContentProvider;
diff --git a/Engine/Plugins/Experimental/UAF/UAFLayering/Source/UAFLayeringUncookedOnly/Internal/UAFLayerStack_EditorData.h b/Engine/Plugins/Experimental/UAF/UAFLayering/Source/UAFLayeringUncookedOnly/Internal/UAFLayerStack_EditorData.h
--- a/Engine/Plugins/Experimental/UAF/UAFLayering/Source/UAFLayeringUncookedOnly/Internal/UAFLayerStack_EditorData.h
+++ b/Engine/Plugins/Experimental/UAF/UAFLayering/Source/UAFLayeringUncookedOnly/Internal/UAFLayerStack_EditorData.h
@@ -16,6 +16,17 @@ class UUAFLayerStack;
 DECLARE_MULTICAST_DELEGATE(FLayerLayoutChangedDelegate);
 DECLARE_MULTICAST_DELEGATE_OneParam(FLayerSelectionChangedDelegate, const TWeakObjectPtr<UUAFLayer>);
 
+namespace UE::UAF::Layering
+{
+	enum class EBaseLayerInclusion
+	{
+		// Includes the base layer when indexing in the layer stack
+		Include,
+		// Excludes the base layer when indexing in the layer stack
+		Exclude
+	};
+}
+
 UCLASS()
 class UUAFLayerStack_EditorData : public UUAFAnimGraph_EditorData, public FEditorUndoClient
 {
@@ -25,31 +36,63 @@ public:
 	UE_API UUAFLayerStack_EditorData();
 	UE_API virtual ~UUAFLayerStack_EditorData() override;
 	
+	// Begin FEditorUndoClient interface
 	virtual void PostUndo(bool bSuccess) override;
 	virtual void PostRedo(bool bSuccess) override;
+	// End FEditorUndoClient interface
+	
+	// Gets the number of total layers 
+	// Per default, this also includes the base layer
+	int32 GetNumLayers(const UE::UAF::Layering::EBaseLayerInclusion IncludeBaseLayer = UE::UAF::Layering::EBaseLayerInclusion::Include) const;
 	
-	int32 GetNumLayers() const;
+	// Gets the index for the given layer, either including the base layer or excluding it
+	int32 GetIndexForLayer(const TObjectPtr<UUAFLayer> InLayer, const UE::UAF::Layering::EBaseLayerInclusion IncludeBaseLayer = UE::UAF::Layering::EBaseLayerInclusion::Include) const;
+	
+	// Gets all layers including the base layer 
 	TArray<TObjectPtr<UUAFLayer>> GetAllLayers() const;
-	int32 GetIndexForLayer(UUAFLayer* Layer) const;
-	int32 GetIndexForMoveableLayer(UUAFLayer* Layer) const;
-	bool IsBaseLayer(UUAFLayer* Layer) const;
-	bool IsLastLayer(UUAFLayer* Layer) const;
-	bool IsLayerNameValid(const FName NewLayerName) const;
-	
-	void MoveLayerToIndex(UUAFLayer* LayerToMove, int32 NewLayerIndex);
-	void MoveLayerUp(UUAFLayer* LayerToMove);
-	void MoveLayerDown(UUAFLayer* LayerToMove);
-	void DeleteLayer(UUAFLayer* LayerToDelete);
-	void SetLayerState(UUAFLayer* Layer, EUAFLayerState LayerState);
-	
-	UE_API void SelectLayer(UUAFLayer* Layer) const;
-	UE_API void ClearSelectedLayer() const;
 	
-	UUAFLayer* CreateDefaultAssetBasedLayer(const FAssetData& InAsset);
+	// Checks if the given layer is the base layer of this layer stack
+	bool IsBaseLayer(const TObjectPtr<UUAFLayer> InLayer) const;
 	
-	void BroadcastLayerStackChanged(EAnimNextEditorDataNotifType InType, bool bRefreshUI = false);
-	void NotifyLayerLayoutChanged() const;
+	// Checks if the given layer is the last layer in this layer stack
+	bool IsLastLayer(const TObjectPtr<UUAFLayer> InLayer) const;
+	
+	// Moves the given layer to a specific index in the layer stack
+	// Does not affect the base layer of the stack 
+	// NewLayerIndex == 0 is equivalent to the first layer under the base layer
+	void MoveLayerToIndex(const TObjectPtr<UUAFLayer> LayerToMove, int32 NewLayerIndex);
+
+	// Decreases the index of the given layer within the layers array of this layer stack and moves it visually up in the stack
+	// If index == 0, this layer will be the next layer under the base layer
+	// Does not affect the base layer of the stack 
+	void MoveLayerUp(const TObjectPtr<UUAFLayer> LayerToMove);
+	
+	// Increases the index of the given layer within the layers array of this layer stack and moves it visually down in the stack
+	// Does not affect the base layer of the stack
+	void MoveLayerDown(const TObjectPtr<UUAFLayer> LayerToMove);
+	
+	// Removes the given layer from the layer stack 
+	void RemoveLayer(const TObjectPtr<UUAFLayer> LayerToDelete);
+	
+	// Checks if NewLayerName is a valid name for the given layer 
+	bool IsLayerNameValid(const TObjectPtr<const UUAFLayer> InLayer, const FName NewLayerName) const;
 	
+	// Generates a unique name for the given layer in the current layer stack
+	// Unless the layer name is already valid, then it will return that 
+	FName GetUniqueNameForLayer(const TObjectPtr<UUAFLayer> InLayer) const;
+	
+	// Set the state of the given layer
+	void SetLayerState(const TObjectPtr<UUAFLayer> Layer, EUAFLayerState LayerState);
+	
+	// Selects the given layer within the layer stack 
+	UE_API void SelectLayer(const TObjectPtr<UUAFLayer> Layer) const;
+
+	// Clears the current selection of the layer stack
+	UE_API void ClearSelectedLayer() const;
+	
+	// Creates a new layer with a default asset and blend provider and adds it on top of the stack
+	TObjectPtr<UUAFLayer> AddDefaultAssetBasedLayer(const FAssetData& InAsset);
+
 protected:
 	// Begin UUAFAnimGraph_EditorData interface
 	UE_API virtual void Initialize(bool bRecompileVM) override;
@@ -59,8 +102,17 @@ protected:
 	// End UUAFAnimGraph_EditorData interface
 	
 	// Begin UObject interface
-	UE_API virtual void PostEditChangeProperty(FPropertyChangedEvent& PropertyChangedEvent) override;
+	UE_API virtual void PostEditChangeChainProperty(FPropertyChangedChainEvent& PropertyChangedEvent) override;
 	// End UObject interface
+
+private:
+	// Broadcasts if anything has changed in the current layer stack, with the option to refresh the layering UI
+	void BroadcastLayerStackChanged(EAnimNextEditorDataNotifType InType, bool bRefreshUI = false);
+	
+	// Notify the layering UI to refresh its state
+	void NotifyLayerLayoutChanged() const;
+	
+	void OnLayerAdded(const TObjectPtr<UUAFLayer> NewLayer);
 	
 public:
 	UPROPERTY(Transient)
diff --git a/Engine/Plugins/Experimental/UAF/UAFLayering/Source/UAFLayeringUncookedOnly/Private/Layers/UAFLayer.cpp b/Engine/Plugins/Experimental/UAF/UAFLayering/Source/UAFLayeringUncookedOnly/Private/Layers/UAFLayer.cpp
--- a/Engine/Plugins/Experimental/UAF/UAFLayering/Source/UAFLayeringUncookedOnly/Private/Layers/UAFLayer.cpp
+++ b/Engine/Plugins/Experimental/UAF/UAFLayering/Source/UAFLayeringUncookedOnly/Private/Layers/UAFLayer.cpp
@@ -138,14 +138,9 @@ void UUAFLayer::RenameLayer(const FName NewName)
 
 bool UUAFLayer::IsLayerNameValid(const FName NewName) const
 {
-	if (NewName == LayerName)
-	{
-		return true;
-	}
-	
 	if (const UUAFLayerStack_EditorData* EditorData = Cast<UUAFLayerStack_EditorData>(GetOuter()))
 	{
-		return EditorData->IsLayerNameValid(NewName);
+		return EditorData->IsLayerNameValid(this, NewName);
 	}
 	
 	return false;
diff --git a/Engine/Plugins/Experimental/UAF/UAFLayering/Source/UAFLayeringUncookedOnly/Private/UAFLayerStack_EditorData.cpp b/Engine/Plugins/Experimental/UAF/UAFLayering/Source/UAFLayeringUncookedOnly/Private/UAFLayerStack_EditorData.cpp
--- a/Engine/Plugins/Experimental/UAF/UAFLayering/Source/UAFLayeringUncookedOnly/Private/UAFLayerStack_EditorData.cpp
+++ b/Engine/Plugins/Experimental/UAF/UAFLayering/Source/UAFLayeringUncookedOnly/Private/UAFLayerStack_EditorData.cpp
@@ -114,16 +114,46 @@ void UUAFLayerStack_EditorData::OnPreCompileGetProgrammaticGraphs(const FRigVMCo
 	}
 }
 
-void UUAFLayerStack_EditorData::PostEditChangeProperty(FPropertyChangedEvent& PropertyChangedEvent)
+void UUAFLayerStack_EditorData::PostEditChangeChainProperty(FPropertyChangedChainEvent& PropertyChangedEvent)
 {
-	Super::PostEditChangeProperty(PropertyChangedEvent);
+	Super::PostEditChangeChainProperty(PropertyChangedEvent);
 
-	if (PropertyChangedEvent.MemberProperty && PropertyChangedEvent.MemberProperty->GetName() == GET_MEMBER_NAME_STRING_CHECKED(UUAFLayerStack_EditorData, Layers))
+	static const FString LayerPropertyString = GET_MEMBER_NAME_STRING_CHECKED(UUAFLayerStack_EditorData, Layers);
+	if (PropertyChangedEvent.MemberProperty && PropertyChangedEvent.MemberProperty->GetName() == LayerPropertyString)
 	{
+		// Check if a new layer got added manually e.g. via the details panel 
+		if (PropertyChangedEvent.ChangeType == EPropertyChangeType::ValueSet)
+		{
+			if (PropertyChangedEvent.PropertyChain.GetTail())
+			{
+				if (const FProperty* const Property = PropertyChangedEvent.PropertyChain.GetTail()->GetValue())
+				{
+					if (Property->GetName() == LayerPropertyString)
+					{
+						const int32 ChangedIndex = PropertyChangedEvent.GetArrayIndex(LayerPropertyString);
+						if (Layers.IsValidIndex(ChangedIndex))
+						{
+							OnLayerAdded(Layers[ChangedIndex]);
+							return;
+						}
+					}
+				}
+			}
+		}
+		
 		NotifyLayerLayoutChanged();
 	}
 }
 
+void UUAFLayerStack_EditorData::OnLayerAdded(const TObjectPtr<UUAFLayer> NewLayer)
+{
+	if (NewLayer)
+	{
+		NewLayer->RenameLayer(GetUniqueNameForLayer(NewLayer));
+		BroadcastLayerStackChanged(EAnimNextEditorDataNotifType::PropertyChanged, true);
+	}
+}
+
 void UUAFLayerStack_EditorData::NotifyLayerLayoutChanged() const
 {
 	OnLayerLayoutChanged.Broadcast();
@@ -138,7 +168,7 @@ TArray<TObjectPtr<UUAFLayer>> UUAFLayerStack_EditorData::GetAllLayers() const
 	return MoveTemp(AllLayers);
 }
 
-void UUAFLayerStack_EditorData::MoveLayerUp(UUAFLayer* LayerToMove)
+void UUAFLayerStack_EditorData::MoveLayerUp(const TObjectPtr<UUAFLayer> LayerToMove)
 {
 	const int32 SourceLayerIndex = Layers.Find(LayerToMove);
 	if (SourceLayerIndex == INDEX_NONE)
@@ -155,7 +185,7 @@ void UUAFLayerStack_EditorData::MoveLayerUp(UUAFLayer* LayerToMove)
 
 	if (SourceLayerIndex == 0)
 	{
-		// can not move this layer up and it is already the uppermost layer after the base layer
+		// can not move this layer up, and it is already the uppermost layer after the base layer
 		return;
 	}
 
@@ -168,7 +198,7 @@ void UUAFLayerStack_EditorData::MoveLayerUp(UUAFLayer* LayerToMove)
 	BroadcastLayerStackChanged(EAnimNextEditorDataNotifType::PropertyChanged, true);
 }
 
-void UUAFLayerStack_EditorData::MoveLayerDown(UUAFLayer* LayerToMove)
+void UUAFLayerStack_EditorData::MoveLayerDown(const TObjectPtr<UUAFLayer> LayerToMove)
 {
 	const int32 SourceLayerIndex = Layers.Find(LayerToMove);
 	if (SourceLayerIndex == INDEX_NONE)
@@ -198,7 +228,7 @@ void UUAFLayerStack_EditorData::MoveLayerDown(UUAFLayer* LayerToMove)
 	BroadcastLayerStackChanged(EAnimNextEditorDataNotifType::PropertyChanged, true);
 }
 
-void UUAFLayerStack_EditorData::DeleteLayer(UUAFLayer* LayerToDelete)
+void UUAFLayerStack_EditorData::RemoveLayer(const TObjectPtr<UUAFLayer> LayerToDelete)
 {
 	if (LayerToDelete)
 	{
@@ -217,8 +247,13 @@ void UUAFLayerStack_EditorData::DeleteLayer(UUAFLayer* LayerToDelete)
 	}
 }
 
-void UUAFLayerStack_EditorData::SetLayerState(UUAFLayer* Layer, EUAFLayerState LayerState)
+void UUAFLayerStack_EditorData::SetLayerState(const TObjectPtr<UUAFLayer> Layer, EUAFLayerState LayerState)
 {
+	if (!Layer)
+	{
+		return;
+	}
+	
 	const FScopedTransaction Transaction(LOCTEXT("SetLayerStateAction", "Set Layer State"));
 	Modify();
 	
@@ -226,7 +261,7 @@ void UUAFLayerStack_EditorData::SetLayerState(UUAFLayer* Layer, EUAFLayerState L
 	BroadcastModified(EAnimNextEditorDataNotifType::PropertyChanged, this);
 }
 
-void UUAFLayerStack_EditorData::SelectLayer(UUAFLayer* Layer) const
+void UUAFLayerStack_EditorData::SelectLayer(const TObjectPtr<UUAFLayer> Layer) const
 {
 	OnLayerSelectionChanged.Broadcast(Layer);
 }
@@ -236,34 +271,56 @@ void UUAFLayerStack_EditorData::ClearSelectedLayer() const
 	OnLayerSelectionChanged.Broadcast(nullptr);
 }
 
-int32 UUAFLayerStack_EditorData::GetIndexForLayer(UUAFLayer* Layer) const 
+int32 UUAFLayerStack_EditorData::GetIndexForLayer(const TObjectPtr<UUAFLayer> InLayer, const UE::UAF::Layering::EBaseLayerInclusion IncludeBaseLayer) const 
 {
-	return GetAllLayers().Find(Layer);
-}
+	switch (IncludeBaseLayer)
+	{
+	case UE::UAF::Layering::EBaseLayerInclusion::Include:
+	{
+		if (InLayer == BaseLayer)
+		{
+			return 0;
+		}
 
-int32 UUAFLayerStack_EditorData::GetIndexForMoveableLayer(UUAFLayer* Layer) const 
-{
-	return Layers.Find(Layer);
+		const int32 LayerIndex = Layers.Find(InLayer);
+		return LayerIndex == INDEX_NONE ? INDEX_NONE : LayerIndex + 1; 
+	}
+	case UE::UAF::Layering::EBaseLayerInclusion::Exclude:
+		return Layers.Find(InLayer);
+	default:
+		// Unhandled case
+		return INDEX_NONE;
+	}
 }
 
-int32 UUAFLayerStack_EditorData::GetNumLayers() const
+
+int32 UUAFLayerStack_EditorData::GetNumLayers(const UE::UAF::Layering::EBaseLayerInclusion IncludeBaseLayer) const
 {
-	return GetAllLayers().Num();
+	switch (IncludeBaseLayer)
+	{
+	case UE::UAF::Layering::EBaseLayerInclusion::Include:
+		return GetAllLayers().Num();
+	case UE::UAF::Layering::EBaseLayerInclusion::Exclude:
+		return Layers.Num();
+	default:
+		// Unhandled case
+		return INDEX_NONE;
+	}
 }
 
-bool UUAFLayerStack_EditorData::IsBaseLayer(UUAFLayer* Layer) const
+bool UUAFLayerStack_EditorData::IsBaseLayer(const TObjectPtr<UUAFLayer> InLayer) const
 {
-	const int32 LayerIndex = GetIndexForLayer(Layer);
+	const int32 LayerIndex = GetIndexForLayer(InLayer);
 	return LayerIndex == 0;
 }
 
-bool UUAFLayerStack_EditorData::IsLastLayer(UUAFLayer* Layer) const
+bool UUAFLayerStack_EditorData::IsLastLayer(const TObjectPtr<UUAFLayer> InLayer) const
 {
-	const int32 LayerIndex = GetIndexForLayer(Layer);
+	const int32 LayerIndex = GetIndexForLayer(InLayer);
 	return LayerIndex == GetAllLayers().Num() - 1;
 }
 
-bool UUAFLayerStack_EditorData::IsLayerNameValid(const FName NewLayerName) const
+bool UUAFLayerStack_EditorData::IsLayerNameValid(const TObjectPtr<const UUAFLayer> InLayer, const FName NewLayerName) const
 {
 	// Empty layer names are not allowed 
 	if (NewLayerName == NAME_None)
@@ -271,10 +328,16 @@ bool UUAFLayerStack_EditorData::IsLayerNameValid(const FName NewLayerName) const
 		return false;
 	}
 	
+	// If renaming a specific layer, it is allowed to keep the current name 
+	if (InLayer && InLayer->GetLayerName() == NewLayerName)
+	{
+		return true;
+	}
+	
 	// Layer names need to be unique within the layer stack 
 	for (const UUAFLayer* Layer : GetAllLayers())
 	{
-		if (!Layer)
+		if (!Layer || Layer == InLayer)
 		{
 			continue;
 		}
@@ -287,12 +350,45 @@ bool UUAFLayerStack_EditorData::IsLayerNameValid(const FName NewLayerName) const
 	return true;
 }
 
-void UUAFLayerStack_EditorData::MoveLayerToIndex(UUAFLayer* LayerToMove, int32 NewLayerIndex)
+FName UUAFLayerStack_EditorData::GetUniqueNameForLayer(const TObjectPtr<UUAFLayer> InLayer) const
+{
+	if (!InLayer)
+	{
+		return NAME_None;
+	}
+	
+	// Check if the current name is already unique 
+	if (IsLayerNameValid(InLayer, InLayer->GetLayerName()))
+	{
+		return InLayer->GetLayerName();
+	}
+	
+	// Otherwise try and name the layer based on its index in the layer stack
+	int32 LayerIndex = GetIndexForLayer(InLayer);
+	if (LayerIndex == INDEX_NONE)
+	{
+		// If the given layer isn't in the layer stack yet, use the next index
+		LayerIndex = GetAllLayers().Num();
+	}
+	
+	static const FText LayerPrefix = LOCTEXT("LayerNamePrefix", "Layer");
+	FName NewLayerName = FName(FString::Printf(TEXT("%s%d"), *LayerPrefix.ToString(), LayerIndex));
+	int32 IterationIndex = 0;
+	while (!IsLayerNameValid(InLayer, NewLayerName))
+	{
+		NewLayerName = FName(FString::Printf(TEXT("%s%d_%d"), *LayerPrefix.ToString(), LayerIndex, IterationIndex));
+		++IterationIndex;
+	}
+	
+	return NewLayerName;
+}
+
+void UUAFLayerStack_EditorData::MoveLayerToIndex(const TObjectPtr<UUAFLayer> LayerToMove, int32 NewLayerIndex)
 {
 	if (LayerToMove && !Layers.IsEmpty())
 	{
 		NewLayerIndex = FMath::Clamp(NewLayerIndex, 0, Layers.Num() - 1);
-		int32 OriginalLayerIndex = GetIndexForMoveableLayer(LayerToMove);
+		int32 OriginalLayerIndex = GetIndexForLayer(LayerToMove, UE::UAF::Layering::EBaseLayerInclusion::Exclude);
 		if (OriginalLayerIndex == NewLayerIndex || OriginalLayerIndex == INDEX_NONE)
 		{
 			return;
@@ -326,7 +422,7 @@ void UUAFLayerStack_EditorData::MoveLayerToIndex(UUAFLayer* LayerToMove, int32 N
 	}
 }
 
-UUAFLayer* UUAFLayerStack_EditorData::CreateDefaultAssetBasedLayer(const FAssetData& InAsset)
+TObjectPtr<UUAFLayer> UUAFLayerStack_EditorData::AddDefaultAssetBasedLayer(const FAssetData& InAsset)
 {
 	const FScopedTransaction Transaction(LOCTEXT("AddLayerAction", "Add Layer"));
 	Modify();
@@ -334,14 +430,14 @@ UUAFLayer* UUAFLayerStack_EditorData::CreateDefaultAssetBasedLayer(const FAssetD
 	UUAFLayer* NewLayer = NewObject<UUAFLayer>(this);
 	if (!NewLayer)
 	{
-		UE_LOG(LogAnimation, Error, TEXT("UUAFLayerStack_EditorData::CreateDefaultAssetBasedLayer: Failed to create a new layer!"));
+		UE_LOG(LogAnimation, Error, TEXT("UUAFLayerStack_EditorData::AddDefaultAssetBasedLayer: Failed to create a new layer!"));
 		return nullptr;
 	}
 	
 	const int32 NewLayerIndex = Layers.Add(NewLayer);
 	if (NewLayerIndex == INDEX_NONE)
 	{
-		UE_LOG(LogAnimation, Error, TEXT(" UUAFLayerStack_EditorData::CreateDefaultAssetBasedLayer: Failed to add new layer to layer array"));
+		UE_LOG(LogAnimation, Error, TEXT(" UUAFLayerStack_EditorData::AddDefaultAssetBasedLayer: Failed to add new layer to layer array"));
 		return nullptr;
 	}
 	
@@ -351,7 +447,7 @@ UUAFLayer* UUAFLayerStack_EditorData::CreateDefaultAssetBasedLayer(const FAssetD
 	NewLayer->SetLayerContentProvider(TInstancedStruct<FUAFLayerAssetProvider>::Make(NewContentProvider));
 	NewLayer->SetLayerBlendProvider(TInstancedStruct<FUAFDefaultBlendProvider>::Make());
 	
-	BroadcastLayerStackChanged(EAnimNextEditorDataNotifType::PropertyChanged, true);
+	OnLayerAdded(NewLayer);
 	
 	return NewLayer;
 }
@@ -372,6 +468,7 @@ void UUAFLayerStack_EditorData::Initialize(bool bRecompileVM)
 	if (BaseLayer == nullptr)
 	{
 		BaseLayer = NewObject<UUAFBaseLayer>(this);
+		OnLayerAdded(BaseLayer);
 	}
 }
 
diff --git a/Engine/Plugins/Experimental/UAF/UAFLayering/Source/UAFLayeringUncookedOnly/Private/Widgets/SLayerStack.cpp b/Engine/Plugins/Experimental/UAF/UAFLayering/Source/UAFLayeringUncookedOnly/Private/Widgets/SLayerStack.cpp
--- a/Engine/Plugins/Experimental/UAF/UAFLayering/Source/UAFLayeringUncookedOnly/Private/Widgets/SLayerStack.cpp
+++ b/Engine/Plugins/Experimental/UAF/UAFLayering/Source/UAFLayeringUncookedOnly/Private/Widgets/SLayerStack.cpp
@@ -102,7 +102,7 @@ void SLayerStack::Construct(const FArguments& InArgs)
 								{
 									if (AllowedClasses.Contains(AssetData.GetClass()))
 									{
-										EditorData->CreateDefaultAssetBasedLayer(AssetData);
+										EditorData->AddDefaultAssetBasedLayer(AssetData);
 									}
 								}
 							})
@@ -121,7 +121,7 @@ void SLayerStack::Construct(const FArguments& InArgs)
 								.VAlign(VAlign_Center)
 								.OnReleased_Lambda([EditorData]()
 									{
-										EditorData->CreateDefaultAssetBasedLayer(FAssetData());
+										EditorData->AddDefaultAssetBasedLayer(FAssetData());
 									})
 								.Content()
 								[
@@ -273,7 +273,7 @@ TSharedRef<ITableRow> SLayerStack::GenerateLayerRow(TWeakObjectPtr<UUAFLayer> It
 					UUAFLayerStack_EditorData* EditorData = UAF::UncookedOnly::FUtils::GetEditorData<UUAFLayerStack_EditorData>(LayerStack.Get());
 					if (EditorData)
 					{
-						int32 TargetIndex = EditorData->GetIndexForMoveableLayer(TargetItem.Get());
+						int32 TargetIndex = EditorData->GetIndexForLayer(TargetItem.Get(), EBaseLayerInclusion::Exclude);
 						if (DropZone == EItemDropZone::BelowItem)
 						{
 							TargetIndex += 1;
@@ -392,12 +392,12 @@ TSharedPtr<SWidget> SLayerStack::OnConstructContextMenu() const
 			}
 		
 			MenuBuilder.AddMenuEntry(
-				LOCTEXT("DeleteLayerLabel", "Delete Layer"),
-				LOCTEXT("DeleteLayerTooltip", "Deletes the selected layer."),
+				LOCTEXT("DeleteLayerLabel", "Remove Layer"),
+				LOCTEXT("DeleteLayerTooltip", "Removes the selected layer."),
 				FSlateIcon(),
 				FUIAction(FExecuteAction::CreateLambda([EditorData, SelectedLayer]()
 					{
-						EditorData->DeleteLayer(SelectedLayer);
+						EditorData->RemoveLayer(SelectedLayer);
 					})),
 				NAME_None,
 				EUserInterfaceActionType::Button);
